Initial Mass UnrealInsights plugin
* Prototype UnrealInsights plugin to ingest and display MassTrace data
* Entity Create, ArchetypeChange and Destroy events aggregated and displayed in list view
* Individual events selectable to better understand archetype changes over recording period
* Timeline markers added with selected events
* Timing regions to delineate processor phases
* Performance is still a work in progress.  Not suitable yet for millions of events.

diff --git a/Engine/Plugins/MassInsights/MassInsights.uplugin b/Engine/Plugins/MassInsights/MassInsights.uplugin
new file mode 100644
--- /dev/null
+++ b/Engine/Plugins/MassInsights/MassInsights.uplugin
@@ -0,0 +1,33 @@
+{
+	"FileVersion": 3,
+	"Version": 1,
+	"VersionName": "1.0",
+	"FriendlyName": "Mass Insights",
+	"Description": "Plugin to gather insights into Mass execution",
+	"Category": "Insights",
+	"CreatedBy": "Epic Games, Inc.",
+	"CreatedByURL": "https://epicgames.com",
+	"EnabledByDefault": true,
+	"CanContainContent": false,
+	"IsBetaVersion": true,
+	"Installed": false,
+	"SupportedPrograms": [ "UnrealInsights" ],
+	"Modules": [
+	{
+		"Name": "MassInsightsAnalysis",
+		"Type": "EditorAndProgram",
+		"LoadingPhase": "Default",
+		"ProgramAllowList": [
+		"UnrealInsights"
+			]
+	},
+	{
+		"Name": "MassInsightsUI",
+		"Type": "EditorAndProgram",
+		"LoadingPhase": "Default",
+		"ProgramAllowList": [
+		"UnrealInsights"
+			]
+	}
+	]
+}
\ No newline at end of file
diff --git a/Engine/Plugins/MassInsights/Source/MassInsightsAnalysis/MassInsightsAnalysis.Build.cs b/Engine/Plugins/MassInsights/Source/MassInsightsAnalysis/MassInsightsAnalysis.Build.cs
new file mode 100644
--- /dev/null
+++ b/Engine/Plugins/MassInsights/Source/MassInsightsAnalysis/MassInsightsAnalysis.Build.cs
@@ -0,0 +1,19 @@
+namespace UnrealBuildTool.Rules
+{
+	public class MassInsightsAnalysis : ModuleRules
+	{
+		public MassInsightsAnalysis(ReadOnlyTargetRules Target) : base(Target)
+		{
+			PublicDependencyModuleNames.AddRange(new string[] { "Core" });
+			PrivateDependencyModuleNames.AddRange(
+			new string[]
+			{
+				"CoreUObject",
+				"TraceLog",
+				"TraceAnalysis",
+				"TraceServices",
+			});
+		}
+	}
+}
+
diff --git a/Engine/Plugins/MassInsights/Source/MassInsightsAnalysis/Private/Analyzers/MassInsightsTraceAnalysis.cpp b/Engine/Plugins/MassInsights/Source/MassInsightsAnalysis/Private/Analyzers/MassInsightsTraceAnalysis.cpp
new file mode 100644
--- /dev/null
+++ b/Engine/Plugins/MassInsights/Source/MassInsightsAnalysis/Private/Analyzers/MassInsightsTraceAnalysis.cpp
@@ -0,0 +1,178 @@
+// Copyright Epic Games, Inc. All Rights Reserved.
+
+#include "MassInsightsTraceAnalysis.h"
+
+#include "Common/ProviderLock.h"
+#include "HAL/LowLevelMemTracker.h"
+#include "TraceServices/Model/AnalysisSession.h"
+#include "Model/MassInsightsPrivate.h"
+
+namespace MassInsightsAnalysis
+{
+
+FMassInsightsTraceAnalyzer::FMassInsightsTraceAnalyzer(TraceServices::IAnalysisSession& InSession,
+                                                         FMassInsightsProvider& InRegionProvider)
+	: Session(InSession)
+	, MassInsightsProvider(InRegionProvider)
+{
+}
+
+void FMassInsightsTraceAnalyzer::OnAnalysisBegin(const FOnAnalysisContext& Context)
+{
+	auto& Builder = Context.InterfaceBuilder;
+	
+	Builder.RouteEvent(RouteId_RegisterMassFragment, "MassTrace", "RegisterMassFragment");
+	Builder.RouteEvent(RouteId_RegisterMassArchetype, "MassTrace", "RegisterMassArchetype");
+	
+	Builder.RouteEvent(RouteId_MassBulkAddEntity, "MassTrace", "MassBulkAddEntity");
+	Builder.RouteEvent(RouteId_MassEntityMoved, "MassTrace", "MassEntityMoved");
+	Builder.RouteEvent(RouteId_MassBulkEntityDestroyed, "MassTrace", "MassBulkEntityDestroyed");
+	
+	Builder.RouteEvent(RouteId_MassPhaseBegin, "MassTrace", "MassPhaseBegin");
+	Builder.RouteEvent(RouteId_MassPhaseEnd, "MassTrace", "MassPhaseEnd");
+}
+
+void FMassInsightsTraceAnalyzer::OnAnalysisEnd()
+{
+	TraceServices::FProviderEditScopeLock RegionProviderScopedLock(static_cast<TraceServices::IEditableProvider&>(MassInsightsProvider));
+	MassInsightsProvider.OnAnalysisSessionEnded();
+}
+
+bool FMassInsightsTraceAnalyzer::OnEvent(uint16 RouteId, EStyle Style, const FOnEventContext& Context)
+{
+	const auto& EventData = Context.EventData;
+	switch (RouteId)
+	{
+		case RouteId_RegisterMassFragment:
+			{
+				uint64 FragmentId = EventData.GetValue<uint64>("FragmentId");
+				FString FragmentName;
+				EventData.GetString("FragmentName", FragmentName);
+
+				uint32 FragmentSize = EventData.GetValue<uint32>("FragmentSize");
+				EFragmentType FragmentType = static_cast<EFragmentType>(EventData.GetValue<uint8>("FragmentType"));
+				FMassFragmentInfo FragmentInfo
+				{
+					.Id = FragmentId,
+					.Name = MoveTemp(FragmentName),
+					.Size = FragmentSize,
+					.Type = FragmentType,
+				};
+
+				TraceServices::FProviderEditScopeLock Lock(MassInsightsProvider);
+				
+				MassInsightsProvider.AddFragment(FragmentInfo);
+				break;
+			}
+		case RouteId_RegisterMassArchetype:
+			{
+				uint64 Id = EventData.GetValue<uint64>("ArchetypeID");
+				const auto FragmentArrayView = EventData.GetArrayView<uint64>("Fragments");
+				FMassArchetypeInfo ArchetypeInfo
+				{
+					.Id = Id,
+				};
+
+				ArchetypeInfo.Fragments.Reserve(FragmentArrayView.Num());
+				{
+					TraceServices::FProviderReadScopeLock ReadLock(MassInsightsProvider);
+					for (int32 Index = 0; Index < FragmentArrayView.Num(); ++Index)
+					{
+						const uint64 FragmentId = FragmentArrayView[Index];
+						ArchetypeInfo.Fragments.Add(MassInsightsProvider.FindFragmentById(FragmentId));
+					}
+					Algo::Sort(ArchetypeInfo.Fragments, [](const FMassFragmentInfo* Lhs, const FMassFragmentInfo* Rhs)
+					{
+						if (Lhs->Type == Rhs->Type)
+						{
+							return Lhs->Name < Rhs->Name;
+						}
+						return Lhs->Type < Rhs->Type;
+					});
+				}
+				{
+					TraceServices::FProviderEditScopeLock EditLock(MassInsightsProvider);
+					MassInsightsProvider.AddArchetype(ArchetypeInfo);
+				}
+				break;
+			}
+		case RouteId_MassBulkAddEntity:
+			{
+				uint64 Cycle = EventData.GetValue<uint64>("Cycle");
+				const auto Entities = EventData.GetArrayView<uint64>("Entities");
+				const auto ArchetypeIDs = EventData.GetArrayView<uint64>("ArchetypeIDs");
+				
+				TraceServices::FProviderEditScopeLock EditLock(MassInsightsProvider);
+				MassInsightsProvider.BulkAddEntity(Context.EventTime.AsSeconds(Cycle), Entities, ArchetypeIDs);
+				break;
+			}
+		case RouteId_MassEntityMoved:
+			{
+				uint64 Cycle = EventData.GetValue<uint64>("Cycle");
+				uint64 Entity = EventData.GetValue<uint64>("Entity");
+				uint64 Archetype = EventData.GetValue<uint64>("NewArchetypeID");
+
+				TraceServices::FProviderEditScopeLock EditLock(MassInsightsProvider);
+				MassInsightsProvider.BulkMoveEntity(
+					Context.EventTime.AsSeconds(Cycle),
+					MakeConstArrayView(&Entity, 1),
+					MakeConstArrayView(&Archetype, 1));
+				break;
+			}
+		case RouteId_MassBulkEntityDestroyed:
+			{
+				uint64 Cycle = EventData.GetValue<uint64>("Cycle");
+				const auto Entities = EventData.GetArrayView<uint64>("Entities");
+
+				TraceServices::FProviderEditScopeLock EditLock(MassInsightsProvider);
+				MassInsightsProvider.BulkDestroyEntity(
+					Context.EventTime.AsSeconds(Cycle),
+					Entities);
+				break;
+			}
+		case RouteId_MassPhaseBegin:
+		{
+			uint64 Cycle = EventData.GetValue<uint64>("Cycle");
+			FString PhaseName;
+			EventData.GetString("PhaseName", PhaseName);
+			
+			const uint64 PhaseId = EventData.GetValue<uint64>("PhaseId", 0);
+			// a missing or 0 RegionID indicates that the region is identified by name only
+			if (PhaseId > 0)
+			{
+				TraceServices::FProviderEditScopeLock RegionProviderScopedLock(MassInsightsProvider);
+				MassInsightsProvider.AppendRegionBegin(*PhaseName, PhaseId, Context.EventTime.AsSeconds(Cycle));
+			}
+			else
+			{
+				TraceServices::FProviderEditScopeLock RegionProviderScopedLock(MassInsightsProvider);
+				MassInsightsProvider.AppendRegionBegin(*PhaseName, Context.EventTime.AsSeconds(Cycle));
+			}
+
+			break;
+		}
+
+		case RouteId_MassPhaseEnd:
+		{
+			uint64 Cycle = EventData.GetValue<uint64>("Cycle");
+			uint64_t PhaseID = EventData.GetValue<uint64>("PhaseId", 0);
+
+			TraceServices::FProviderEditScopeLock RegionProviderScopedLock(MassInsightsProvider);
+			if (PhaseID > 0)
+			{
+				MassInsightsProvider.AppendRegionEnd(PhaseID, Context.EventTime.AsSeconds(Cycle));
+			}
+			else
+			{
+				FString PhaseName = TEXT("Invalid");
+				EventData.GetString("PhaseName", PhaseName);
+				MassInsightsProvider.AppendRegionEnd(*PhaseName, Context.EventTime.AsSeconds(Cycle));
+			}
+			break;
+		}
+	}
+
+	return true;
+}
+
+} // namespace MassInsightsAnalysis
diff --git a/Engine/Plugins/MassInsights/Source/MassInsightsAnalysis/Private/Analyzers/MassInsightsTraceAnalysis.h b/Engine/Plugins/MassInsights/Source/MassInsightsAnalysis/Private/Analyzers/MassInsightsTraceAnalysis.h
new file mode 100644
--- /dev/null
+++ b/Engine/Plugins/MassInsights/Source/MassInsightsAnalysis/Private/Analyzers/MassInsightsTraceAnalysis.h
@@ -0,0 +1,46 @@
+// Copyright Epic Games, Inc. All Rights Reserved.
+
+#pragma once
+
+#include "Trace/Analyzer.h"
+
+namespace TraceServices
+{
+	class IAnalysisSession;
+}
+
+namespace MassInsightsAnalysis
+{
+	
+class FMassInsightsProvider;
+
+/*
+ *  Analyzer for Timing Regions, parses the TimingRegionBegin and TimingRegionEnd events from the 'Misc' Logger 
+ */
+class FMassInsightsTraceAnalyzer
+	: public UE::Trace::IAnalyzer
+{
+public:
+	FMassInsightsTraceAnalyzer(TraceServices::IAnalysisSession& Session,
+	                            FMassInsightsProvider& RegionProvider);
+	virtual void OnAnalysisBegin(const FOnAnalysisContext& Context) override;
+	virtual void OnAnalysisEnd() override;
+	virtual bool OnEvent(uint16 RouteId, EStyle Style, const FOnEventContext& Context) override;
+
+private:
+	enum : uint16
+	{
+		RouteId_RegisterMassFragment,
+		RouteId_RegisterMassArchetype,
+		RouteId_MassBulkAddEntity,
+		RouteId_MassEntityMoved,
+		RouteId_MassBulkEntityDestroyed,
+		RouteId_MassPhaseBegin,
+		RouteId_MassPhaseEnd
+	};
+
+	TraceServices::IAnalysisSession& Session;
+	FMassInsightsProvider& MassInsightsProvider;
+};
+
+} // namespace MassInsightsAnalysis
diff --git a/Engine/Plugins/MassInsights/Source/MassInsightsAnalysis/Private/MassInsightsAnalysisModule.cpp b/Engine/Plugins/MassInsights/Source/MassInsightsAnalysis/Private/MassInsightsAnalysisModule.cpp
new file mode 100644
--- /dev/null
+++ b/Engine/Plugins/MassInsights/Source/MassInsightsAnalysis/Private/MassInsightsAnalysisModule.cpp
@@ -0,0 +1,39 @@
+// Copyright Epic Games, Inc. All Rights Reserved.
+
+#include "MassInsightsAnalysisModule.h"
+
+#include "Model/MassInsightsPrivate.h"
+#include "Analyzers/MassInsightsTraceAnalysis.h"
+#include "Modules/ModuleManager.h"
+#include "Features/IModularFeatures.h"
+
+namespace MassInsightsAnalysis
+{
+
+void FMassInsightsAnalysisModule::OnAnalysisBegin(TraceServices::IAnalysisSession& InSession)
+{
+	TSharedPtr<FMassInsightsProvider> MassInsightsProvider = MakeShared<FMassInsightsProvider>(InSession);
+	InSession.AddProvider(GetMassInsightsProviderName(), MassInsightsProvider);
+
+	InSession.AddAnalyzer(new FMassInsightsTraceAnalyzer(InSession, *MassInsightsProvider));
+}
+
+void FMassInsightsAnalysisModule::GetModuleInfo(TraceServices::FModuleInfo& OutModuleInfo)
+{
+	OutModuleInfo.Name = TEXT("MassInsightsProvider");
+	OutModuleInfo.DisplayName = TEXT("MassInsights");
+}
+
+void FMassInsightsAnalysisModule::StartupModule()
+{
+	IModularFeatures::Get().RegisterModularFeature(TraceServices::ModuleFeatureName, this);
+}
+
+void FMassInsightsAnalysisModule::ShutdownModule()
+{
+	IModularFeatures::Get().UnregisterModularFeature(TraceServices::ModuleFeatureName, this);
+}
+
+} //namespace MassInsightsAnalysis
+
+IMPLEMENT_MODULE(MassInsightsAnalysis::FMassInsightsAnalysisModule, MassInsightsAnalysis);
diff --git a/Engine/Plugins/MassInsights/Source/MassInsightsAnalysis/Private/MassInsightsAnalysisModule.h b/Engine/Plugins/MassInsights/Source/MassInsightsAnalysis/Private/MassInsightsAnalysisModule.h
new file mode 100644
--- /dev/null
+++ b/Engine/Plugins/MassInsights/Source/MassInsightsAnalysis/Private/MassInsightsAnalysisModule.h
@@ -0,0 +1,29 @@
+// Copyright Epic Games, Inc. All Rights Reserved.
+
+#pragma once
+
+#include "Modules/ModuleInterface.h"
+#include "TraceServices/ModuleService.h"
+
+namespace MassInsightsAnalysis
+{
+/**
+ *  Class to create and register the Analyzer and Provider for MassInsights.
+ *  Implements both a plugin module and the TraceServices::IModule modular feature.
+ */
+class FMassInsightsAnalysisModule : public IModuleInterface, TraceServices::IModule
+{
+public:
+	//~ Begin IModuleInterface interface
+	virtual void StartupModule() override;
+	virtual void ShutdownModule() override;
+	//~ End IModuleInterface interface
+
+	public:
+	//~ Begin TraceServices::IModule interface
+	virtual void GetModuleInfo(TraceServices::FModuleInfo& OutModuleInfo) override;
+	virtual void OnAnalysisBegin(TraceServices::IAnalysisSession& Session) override;
+	//~ End TraceServices::IModule interface
+};
+	
+} //namespace MassInsightsAnalysis
diff --git a/Engine/Plugins/MassInsights/Source/MassInsightsAnalysis/Private/Model/MassInsights.cpp b/Engine/Plugins/MassInsights/Source/MassInsightsAnalysis/Private/Model/MassInsights.cpp
new file mode 100644
--- /dev/null
+++ b/Engine/Plugins/MassInsights/Source/MassInsightsAnalysis/Private/Model/MassInsights.cpp
@@ -0,0 +1,425 @@
+#include "MassInsightsAnalysis/Model/MassInsights.h"
+#include "MassInsightsPrivate.h"
+
+#include "Algo/ForEach.h"
+#include "Internationalization/Internationalization.h"
+
+DEFINE_LOG_CATEGORY(LogMassInsights)
+
+namespace MassInsightsAnalysis
+{
+
+thread_local TraceServices::FProviderLock::FThreadLocalState GMassInsightsProviderLockState;
+	
+FMassInsightsProvider::FMassInsightsProvider(TraceServices::IAnalysisSession& InSession)
+	: Session(InSession)
+	, EntityEvents(InSession.GetLinearAllocator(), Constants::EntityEventsPageSize)
+{
+}
+
+int32 FMassInsightsProvider::GetFragmentCount() const
+{
+	ReadAccessCheck();
+	return FragmentInfoByID.Num();
+}
+
+const FMassFragmentInfo* FMassInsightsProvider::FindFragmentById(uint64 FragmentId) const
+{
+	ReadAccessCheck();
+	
+	const FMassFragmentInfo* const* Result = FragmentInfoByID.Find(FragmentId);
+	if (Result != nullptr)
+	{
+		return *Result;
+	}
+	return nullptr;
+}
+
+const FMassArchetypeInfo* FMassInsightsProvider::FindArchetypeById(uint64 ArchetypeId) const
+{
+	ReadAccessCheck();
+
+	const FMassArchetypeInfo* const* Result = ArchetypeByID.Find(ArchetypeId);
+	if (Result != nullptr)
+	{
+		return *Result;
+	}
+	return nullptr;
+}
+
+void FMassInsightsProvider::EnumerateFragments(TFunctionRef<void(const FMassFragmentInfo& FragmentInfo, int32 Index)> Callback, int32 BeginIndex) const
+{
+	ReadAccessCheck();
+	for (int32 Index = BeginIndex, End = FragmentInfos.Num(); Index < End; Index++)
+	{
+		const FMassFragmentInfo& FragmentInfo = FragmentInfos[Index];
+		Callback(FragmentInfo, Index);
+	}
+}
+
+uint64 FMassInsightsProvider::GetEntityEventCount() const
+{
+	ReadAccessCheck();
+
+	return EntityEvents.Num();
+}
+
+TValueOrError<FMassEntityEventRecord, void> FMassInsightsProvider::GetEntityEvent(uint64 EventIndex) const
+{
+	ReadAccessCheck();
+	
+	if (EventIndex < EntityEvents.Num())
+	{
+		return MakeValue(EntityEvents[EventIndex]);
+	}
+
+	return MakeError();
+}
+
+void FMassInsightsProvider::EnumerateEntityEvents(
+	uint64 StartIndex,
+	uint64 Count,
+	TFunctionRef<void(const FMassEntityEventRecord&, uint64 EventIndex)> Callback) const
+{
+	ReadAccessCheck();
+
+	uint64 EndIndex = FPlatformMath::Min(StartIndex + Count, EntityEvents.Num());
+	for (uint64 Index = StartIndex; Index < EndIndex; Index++)
+	{
+		Callback(EntityEvents[Index], Index);
+	}
+}
+
+uint64 FMassInsightsProvider::GetRegionCount() const
+{
+	ReadAccessCheck();
+
+	uint64 RegionCount = 0;
+	for (const FMassInsightsLane& Lane : Lanes)
+	{
+		RegionCount += Lane.Num();
+	}
+	return RegionCount;
+}
+
+int32 FMassInsightsProvider::GetLaneCount() const
+{
+	ReadAccessCheck();
+	return Lanes.Num();
+}
+
+const FMassInsightsLane* FMassInsightsProvider::GetLane(int32 index) const
+{
+	ReadAccessCheck();
+
+	if (index < Lanes.Num())
+	{
+		return &(Lanes[index]);
+	}
+	return nullptr;
+}
+
+void FMassInsightsProvider::AppendRegionBegin(const TCHAR* Name, double Time)
+{
+	EditAccessCheck();
+	// Regions identified by Name don't have an ID
+	AppendRegionBegin(Name, 0, Time);
+}
+
+void FMassInsightsProvider::AddFragment(const FMassFragmentInfo& FragmentInfo)
+{
+	EditAccessCheck();
+
+	if (FragmentInfoByID.Find(FragmentInfo.Id) == nullptr)
+	{
+		uint64 Id = FragmentInfo.Id;
+		int32 AllocatedIndex = FragmentInfos.AddElement(FragmentInfo);
+		const FMassFragmentInfo* AllocatedFragment = &FragmentInfos[AllocatedIndex];
+		FragmentInfoByID.Add(Id, AllocatedFragment);
+	}
+}
+
+void FMassInsightsProvider::AddArchetype(const FMassArchetypeInfo& ArchetypeInfo)
+{
+	EditAccessCheck();
+
+	if (ArchetypeByID.Find(ArchetypeInfo.Id) == nullptr)
+	{
+		uint32 Id = ArchetypeInfo.Id;
+		int32 AllocatedIndex = ArchetypeInfos.AddElement(ArchetypeInfo);
+		const FMassArchetypeInfo* AllocatedArchetype = &ArchetypeInfos[AllocatedIndex];
+		ArchetypeByID.Add(Id, AllocatedArchetype);
+	}
+}
+
+void FMassInsightsProvider::BulkAddEntity(double Time, TConstArrayView<uint64> Entities, TConstArrayView<uint64> ArchetypeIDs)
+{
+	EditAccessCheck();
+
+	FMassEntityEventRecord Row;
+	Row.Time = Time;
+	Row.Operation = EMassEntityEventType::Created;
+
+	for (int32 Index = 0, End = Entities.Num(); Index < End; Index++)
+	{
+		Row.ArchetypeID = ArchetypeIDs[Index];
+		Row.Entity = Entities[Index];
+		EntityEvents.EmplaceBack(Row);
+	}
+}
+	
+void FMassInsightsProvider::BulkMoveEntity(double Time, TConstArrayView<uint64> Entities, TConstArrayView<uint64> ArchetypeIDs)
+{
+	EditAccessCheck();
+
+	FMassEntityEventRecord Row;
+	Row.Time = Time;
+	Row.Operation = EMassEntityEventType::ArchetypeChange;
+	
+	for (int32 Index = 0, End = Entities.Num(); Index < End; Index++)
+	{
+		Row.ArchetypeID = ArchetypeIDs[Index];
+		Row.Entity = Entities[Index];
+		EntityEvents.EmplaceBack(Row);
+	}
+}
+	
+void FMassInsightsProvider::BulkDestroyEntity(double Time, TConstArrayView<uint64> Entities)
+{
+	EditAccessCheck();
+
+	FMassEntityEventRecord Row;
+	Row.Time = Time;
+	Row.Operation = EMassEntityEventType::Destroyed;
+	Row.ArchetypeID = 0;
+	
+	for (int32 Index = 0, End = Entities.Num(); Index < End; Index++)
+	{
+		Row.Entity = Entities[Index];
+		EntityEvents.EmplaceBack(Row);
+	}
+}
+
+void FMassInsightsProvider::AppendRegionBegin(const TCHAR* Name, uint64 ID, double Time)
+{
+	EditAccessCheck();
+	// lookup by ID if the region has an ID
+	FMassInsights** OpenRegion = ID ? OpenRegionsByID.Find(ID) : OpenRegionsByName.Find(Name);
+	
+	if (OpenRegion)
+	{
+		++NumWarnings;
+		if (NumWarnings <= MaxWarningMessages)
+		{
+			UE_LOG(LogMassInsights, Warning, TEXT("[Regions] A region begin event (%s) was encountered while a region with same name is already open."), Name)
+		}
+	}
+	else
+	{
+		FMassInsights Region;
+		Region.BeginTime = Time;
+		Region.Text = Session.StoreString(Name);
+		Region.ID = ID;
+		Region.Depth = CalculateRegionDepth(Region);
+
+		if (Region.Depth == Lanes.Num())
+		{
+			Lanes.Emplace(Session.GetLinearAllocator());
+		}
+
+		Lanes[Region.Depth].Regions.EmplaceBack(Region);
+		FMassInsights* NewOpenRegion = &(Lanes[Region.Depth].Regions.Last());
+		
+		if (ID)
+		{
+			OpenRegionsByID.Add(Region.ID, NewOpenRegion);
+		} else
+		{
+			OpenRegionsByName.Add(Region.Text, NewOpenRegion);
+		}
+		UpdateCounter++;
+	}
+
+	// Update session time
+	{
+		TraceServices::FAnalysisSessionEditScope _(Session);
+		Session.UpdateDurationSeconds(Time);
+	}
+}
+
+void FMassInsightsProvider::AppendRegionEnd(const uint64 ID, double Time)
+{
+	EditAccessCheck();
+
+	FMassInsights** OpenRegionPos = OpenRegionsByID.Find(ID);
+	FMassInsights* OpenRegion = nullptr;
+	if (OpenRegionPos)
+		OpenRegion = *OpenRegionPos;
+
+	AppendRegionEndInternal(OpenRegion, Time);
+}
+	
+void FMassInsightsProvider::AppendRegionEnd(const TCHAR* Name, double Time)
+{
+	EditAccessCheck();
+
+	FMassInsights** OpenRegionPos = OpenRegionsByName.Find(Name);
+	FMassInsights* OpenRegion = nullptr;
+	if (OpenRegionPos)
+	{
+		OpenRegion = *OpenRegionPos;		
+	}
+	else
+	{
+		++NumWarnings;
+		if (NumWarnings <= MaxWarningMessages)
+		{
+			UE_LOG(LogMassInsights, Warning, TEXT("[Regions] A region end event (%s) was encountered without having seen a matching region start event first."), Name)
+		}
+	}
+
+	AppendRegionEndInternal(OpenRegion, Time);
+}
+
+void FMassInsightsProvider::AppendRegionEndInternal(FMassInsights* OpenRegion, double Time)
+{
+	if (OpenRegion)
+	{
+		OpenRegion->EndTime = Time;
+
+		if (OpenRegion->ID)
+		{
+			OpenRegionsByID.Remove(OpenRegion->ID);
+		}
+		OpenRegionsByName.Remove(OpenRegion->Text);
+		UpdateCounter++;
+	}
+
+	// Update session time
+	{
+		TraceServices::FAnalysisSessionEditScope _(Session);
+		Session.UpdateDurationSeconds(Time);
+	}
+}
+	
+void FMassInsightsProvider::OnAnalysisSessionEnded()
+{
+	EditAccessCheck();
+
+	auto printOpenRegionMessage = [this](const auto& KV)
+	{
+		const FMassInsights* Region = KV.Value;
+		++NumWarnings;
+		if (NumWarnings <= MaxWarningMessages)
+		{
+			UE_LOG(LogMassInsights, Warning, TEXT("[Regions] A region begin event (%s) was never closed."), Region->Text)
+		}
+	};
+	Algo::ForEach(OpenRegionsByID, printOpenRegionMessage);
+	Algo::ForEach(OpenRegionsByName, printOpenRegionMessage);
+
+	if (NumWarnings > 0 || NumErrors > 0)
+	{
+		UE_LOG(LogMassInsights, Error, TEXT("[Regions] %u warnings; %u errors"), NumWarnings, NumErrors);
+	}
+
+	const uint64 TotalRegionCount = GetRegionCount();
+	UE_LOG(LogMassInsights, Log, TEXT("[Regions] Analysis completed (%llu regions, %d lanes)."), TotalRegionCount, Lanes.Num());
+}
+
+int32 FMassInsightsProvider::CalculateRegionDepth(const FMassInsights& Region) const
+{
+	constexpr int32 DepthLimit = 100;
+
+	int32 NewDepth = 0;
+
+	// Find first free lane/depth
+	while (NewDepth < DepthLimit)
+	{
+		if (!Lanes.IsValidIndex(NewDepth))
+		{
+			break;
+		}
+
+		const FMassInsights& LastRegion = Lanes[NewDepth].Regions.Last();
+		if (LastRegion.EndTime <= Region.BeginTime)
+		{
+			break;
+		}
+		NewDepth++;
+	}
+
+	ensureMsgf(NewDepth < DepthLimit, TEXT("Regions are nested too deep."));
+
+	return NewDepth;
+}
+
+void FMassInsightsProvider::EnumerateLanes(TFunctionRef<void(const FMassInsightsLane&, int32)> Callback) const
+{
+	ReadAccessCheck();
+
+	for (int32 LaneIndex = 0; LaneIndex < Lanes.Num(); ++LaneIndex)
+	{
+		Callback(Lanes[LaneIndex], LaneIndex);
+	}
+}
+
+bool FMassInsightsProvider::EnumerateRegions(double IntervalStart, double IntervalEnd, TFunctionRef<bool(const FMassInsights&)> Callback) const
+{
+	ReadAccessCheck();
+
+	if (IntervalStart > IntervalEnd)
+	{
+		return false;
+	}
+
+	for (const FMassInsightsLane& Lane : Lanes)
+	{
+		if (!Lane.EnumerateRegions(IntervalStart, IntervalEnd, Callback))
+		{
+			return false;
+		}
+	}
+
+	return true;
+}
+
+bool FMassInsightsLane::EnumerateRegions(double IntervalStart, double IntervalEnd, TFunctionRef<bool(const FMassInsights&)> Callback) const
+{
+	const FInt32Interval OverlapRange = GetElementRangeOverlappingGivenRange<FMassInsights>(Regions, IntervalStart, IntervalEnd,
+		[](const FMassInsights& r) { return r.BeginTime; },
+		[](const FMassInsights& r) { return r.EndTime; });
+
+	if (OverlapRange.Min == -1)
+	{
+		return true;
+	}
+
+	for (int32 Index = OverlapRange.Min; Index <= OverlapRange.Max; ++Index)
+	{
+		if (!Callback(Regions[Index]))
+		{
+			return false;
+		}
+	}
+
+	return true;
+}
+	
+FName GetMassInsightsProviderName()
+{
+	static const FName Name("MassInsightsProvider");
+	return Name;
+}
+
+const IMassInsightsProvider& ReadMassInsightsProvider(const TraceServices::IAnalysisSession& Session)
+{
+	return *Session.ReadProvider<IMassInsightsProvider>(GetMassInsightsProviderName());
+}
+
+IEditableMassInsightsProvider& EditMassInsightsProvider(TraceServices::IAnalysisSession& Session)
+{
+	return *Session.EditProvider<IEditableMassInsightsProvider>(GetMassInsightsProviderName());
+}
+
+} // namespace MassInsightsAnalysis
diff --git a/Engine/Plugins/MassInsights/Source/MassInsightsAnalysis/Private/Model/MassInsightsPrivate.h b/Engine/Plugins/MassInsights/Source/MassInsightsAnalysis/Private/Model/MassInsightsPrivate.h
new file mode 100644
--- /dev/null
+++ b/Engine/Plugins/MassInsights/Source/MassInsightsAnalysis/Private/Model/MassInsightsPrivate.h
@@ -0,0 +1,128 @@
+#pragma once
+
+#include "Common/ProviderLock.h"
+#include "Containers/ChunkedArray.h"
+#include "Containers/PagedArray.h"
+#include "Templates/SharedPointer.h"
+#include "MassInsightsAnalysis/Model/MassInsights.h"
+
+DECLARE_LOG_CATEGORY_EXTERN(LogMassInsights, Display, All)
+
+namespace TraceServices
+{
+	class FAnalysisSessionLock;
+	class FStringStore;
+}
+
+namespace MassInsightsAnalysis
+{
+
+extern thread_local TraceServices::FProviderLock::FThreadLocalState GMassInsightsProviderLockState;
+
+class FMassInsightsProvider
+	: public IMassInsightsProvider
+	, public IEditableMassInsightsProvider
+{
+	struct Constants
+	{
+		static constexpr uint64 EntityEventsPageSize = 65536;
+	};
+public:
+	explicit FMassInsightsProvider(TraceServices::IAnalysisSession& Session);
+	virtual ~FMassInsightsProvider() override {}
+
+	//////////////////////////////////////////////////
+	// Read operations
+
+	virtual void BeginRead() const override       { Lock.BeginRead(GMassInsightsProviderLockState); }
+	virtual void EndRead() const override         { Lock.EndRead(GMassInsightsProviderLockState); }
+	virtual void ReadAccessCheck() const override { Lock.ReadAccessCheck(GMassInsightsProviderLockState); }
+
+	virtual int32 GetFragmentCount() const override;
+	virtual const FMassFragmentInfo* FindFragmentById(uint64 FragmentId) const override;
+	virtual const FMassArchetypeInfo* FindArchetypeById(uint64 ArchetypeId) const override;
+	
+	virtual void EnumerateFragments(TFunctionRef<void(const FMassFragmentInfo& FragmentInfo, int32 Index)> Callback, int32 BeginIndex) const override;
+
+	virtual uint64 GetEntityEventCount() const override;
+	virtual TValueOrError<FMassEntityEventRecord, void> GetEntityEvent(uint64 EventIndex) const override;
+	virtual void EnumerateEntityEvents(
+		uint64 StartIndex,
+		uint64 Count,
+		TFunctionRef<void(const FMassEntityEventRecord&, uint64 /*EventIndex*/)> Callback) const override;
+	
+	virtual uint64 GetRegionCount() const override;
+	virtual int32 GetLaneCount()  const override;
+
+	virtual const FMassInsightsLane* GetLane(int32 Index) const override;
+
+	virtual bool EnumerateRegions(double IntervalStart, double IntervalEnd, TFunctionRef<bool(const FMassInsights&)> Callback) const override;
+	virtual void EnumerateLanes(TFunctionRef<void(const FMassInsightsLane&, const int32)> Callback) const override;
+
+	virtual uint64 GetUpdateCounter() const override { ReadAccessCheck(); return UpdateCounter; }
+
+	//////////////////////////////////////////////////
+	// Edit operations
+
+	virtual void BeginEdit() const override       { Lock.BeginWrite(GMassInsightsProviderLockState); }
+	virtual void EndEdit() const override         { Lock.EndWrite(GMassInsightsProviderLockState); }
+	virtual void EditAccessCheck() const override { Lock.WriteAccessCheck(GMassInsightsProviderLockState); }
+
+	virtual void AddFragment(const FMassFragmentInfo& FragmentInfo) override;
+	virtual void AddArchetype(const FMassArchetypeInfo& ArchetypeInfo) override;
+	
+	virtual void BulkAddEntity(double Time, TConstArrayView<uint64> Entities, TConstArrayView<uint64> ArchetypeIDs) override;
+	virtual void BulkMoveEntity(double Time, TConstArrayView<uint64> Entities, TConstArrayView<uint64> ArchetypeIDs) override;
+	virtual void BulkDestroyEntity(double Time, TConstArrayView<uint64> Entities) override;
+	
+	virtual void AppendRegionBegin(const TCHAR* Name, uint64 ID, double Time) override;
+	virtual void AppendRegionBegin(const TCHAR* Name, double Time) override;
+	virtual void AppendRegionEnd(const TCHAR* ID, double Time) override;
+	virtual void AppendRegionEnd(const uint64 ID, double Time) override;
+
+	virtual void OnAnalysisSessionEnded() override;
+
+	//////////////////////////////////////////////////
+
+private:
+	// Update the depth member of a region to allow overlapping regions to be displayed on separate lanes.
+	int32 CalculateRegionDepth(const FMassInsights& Item) const;
+	void AppendRegionEndInternal(FMassInsights* OpenRegion, double Time);
+
+private:
+	mutable TraceServices::FProviderLock Lock;
+
+	TraceServices::IAnalysisSession&	Session;
+
+	TMap<uint64, const FMassFragmentInfo*> FragmentInfoByID;
+	TChunkedArray<FMassFragmentInfo> FragmentInfos;
+
+	// Allocation of ranges of fragments, mostly for ArchetypeInfo
+	TChunkedArray<const FMassFragmentInfo*> FragmentInfoRanges;
+
+	TMap<uint32, const FMassArchetypeInfo*> ArchetypeByID;
+	TChunkedArray<FMassArchetypeInfo> ArchetypeInfos;
+	
+	// Timeseries updates
+	
+	// Sorted by Cycle
+	TraceServices::TPagedArray<FMassEntityEventRecord> EntityEvents;
+
+	// Open regions inside lanes
+	TMap<FStringView, FMassInsights*> OpenRegionsByName;
+	TMap<uint64_t, FMassInsights*> OpenRegionsByID;
+
+	// Closed regions
+	TArray<FMassInsightsLane> Lanes;
+
+	// Counter incremented each time region data changes during analysis
+	uint64 UpdateCounter = -1;
+
+	static constexpr uint32 MaxWarningMessages = 100;
+	static constexpr uint32 MaxErrorMessages = 100;
+
+	uint32 NumWarnings = 0;
+	uint32 NumErrors = 0;
+};
+
+} // namespace MassInsightsAnalysis
diff --git a/Engine/Plugins/MassInsights/Source/MassInsightsAnalysis/Public/MassInsightsAnalysis/Model/MassInsights.h b/Engine/Plugins/MassInsights/Source/MassInsightsAnalysis/Public/MassInsightsAnalysis/Model/MassInsights.h
new file mode 100644
--- /dev/null
+++ b/Engine/Plugins/MassInsights/Source/MassInsightsAnalysis/Public/MassInsightsAnalysis/Model/MassInsights.h
@@ -0,0 +1,212 @@
+#pragma once
+
+#include "Common/PagedArray.h" // TraceServices
+#include "HAL/Platform.h"
+#include "Templates/Function.h"
+#include "Templates/ValueOrError.h"
+#include "TraceServices/Model/AnalysisSession.h"
+#include "UObject/NameTypes.h"
+
+namespace MassInsightsAnalysis
+{
+
+struct MASSINSIGHTSANALYSIS_API FMassInsights
+{
+	double BeginTime = std::numeric_limits<double>::infinity();
+	double EndTime = std::numeric_limits<double>::infinity();
+	const TCHAR* Text = nullptr;
+	// ID will be zero if the region is identified by Name only
+	uint64 ID = 0;
+	int32 Depth = -1;
+};
+
+class MASSINSIGHTSANALYSIS_API FMassInsightsLane
+{
+	friend class FMassInsightsProvider;
+
+public:
+	FMassInsightsLane(TraceServices::ILinearAllocator& InAllocator) : Regions(InAllocator, 512) {}
+
+	int32 Num() const { return static_cast<int32>(Regions.Num()); }
+
+	/**
+	 * Call Callback for every region overlapping the interval defined by IntervalStart and IntervalEnd
+	 * @param Callback a callback called for each region. Return false to abort iteration.
+	 * @returns true if the enumeration finished, false if it was aborted by the callback returning false
+	 */
+	bool EnumerateRegions(double IntervalStart, double IntervalEnd, TFunctionRef<bool(const FMassInsights&)> Callback) const;
+
+private:
+	TraceServices::TPagedArray<FMassInsights> Regions;
+};
+
+enum class EFragmentType : uint8
+{
+	Unknown,
+	Fragment,
+	Tag,
+	Shared
+};
+
+struct MASSINSIGHTSANALYSIS_API FMassFragmentInfo
+{
+	uint64 Id;
+	FString Name;
+	uint32 Size;
+	EFragmentType Type;
+
+	const TCHAR* GetName() const { return *Name; }
+};
+
+struct MASSINSIGHTSANALYSIS_API FMassArchetypeInfo
+{
+	uint64 Id;
+	TArray<const FMassFragmentInfo*> Fragments;
+
+	TConstArrayView<const FMassFragmentInfo*> GetFragments() const
+	{
+		return MakeArrayView(Fragments);
+	}
+};
+	
+enum class EMassEntityEventType : uint8_t
+{
+	Created,
+	ArchetypeChange,
+	Destroyed
+};
+
+struct FMassEntityEventRecord
+{
+	double Time;
+	uint64 Entity;
+	uint64 ArchetypeID;
+	EMassEntityEventType Operation;
+};
+
+class MASSINSIGHTSANALYSIS_API IMassInsightsProvider
+	: public TraceServices::IProvider
+{
+protected:
+	
+public:
+	virtual ~IMassInsightsProvider() override = default;
+
+	virtual int32 GetFragmentCount() const = 0;
+	virtual const FMassFragmentInfo* FindFragmentById(uint64 FragmentId) const = 0;
+	
+	virtual const FMassArchetypeInfo* FindArchetypeById(uint64 ArchetypeId) const = 0;
+
+	virtual void EnumerateFragments(TFunctionRef<void(const FMassFragmentInfo& FragmentInfo, int32 Index)> Callback, int32 BeginIndex) const = 0;
+
+	virtual TValueOrError<FMassEntityEventRecord, void> GetEntityEvent(uint64 EventIndex) const = 0;
+
+	virtual uint64 GetEntityEventCount() const = 0;
+	// Enumerate upto Count number of events starting at the StartIndex
+	// Enumeration will end early if there is not enough events or if Callback returns false
+	virtual void EnumerateEntityEvents(
+		uint64 StartIndex,
+		uint64 Count,
+		TFunctionRef<void(const FMassEntityEventRecord&, uint64 EventIndex)> Callback) const = 0;
+
+	/**
+	 * @return the amount of currently known regions (including open-ended ones)
+	 */
+	virtual uint64 GetRegionCount() const = 0;
+
+	/**
+	 * @return the number of lanes
+	 */
+	virtual int32 GetLaneCount() const = 0;
+
+	/**
+	 * Direct access to a certain lane at a given index/depth.
+	 * The pointer is valid only in the current read scope.
+	 * @return a pointer to the lane at the specified depth index or nullptr if Index > GetLaneCount()-1
+	 */
+	virtual const FMassInsightsLane* GetLane(int32 Index) const = 0;
+
+	/**
+	 * Enumerates all regions that overlap a certain time interval. Will enumerate by depth but does not expose lanes.
+	 * @param Callback a callback called for each region. Return false to abort iteration.
+	 * @returns true if the enumeration finished, false if it was aborted by the callback returning false
+	 */
+	virtual bool EnumerateRegions(double  IntervalStart, double IntervalEnd, TFunctionRef<bool(const FMassInsights&)> Callback) const = 0;
+
+	/**
+	 * Will call Callback(Lane, Depth) for each lane in order.
+	 */
+	virtual void EnumerateLanes(TFunctionRef<void(const FMassInsightsLane&, const int32)> Callback) const = 0;
+
+	/**
+	 * @return A monotonically increasing counter that that changes each time new data is added to the provider.
+	 * This can be used to detect when to update any (UI-)state dependent on the provider during analysis.
+	 */
+	virtual uint64 GetUpdateCounter() const = 0;
+};
+
+/*
+* The interface to a provider that can consume mutations of region events from a session.
+*/
+class IEditableMassInsightsProvider
+	: public TraceServices::IEditableProvider
+{
+public:
+	virtual ~IEditableMassInsightsProvider() override = default;
+	
+	virtual void AddFragment(const FMassFragmentInfo& FragmentInfo) = 0;
+	virtual void AddArchetype(const FMassArchetypeInfo& ArchetypeInfo) = 0;
+	
+	virtual void BulkAddEntity(double Time, TConstArrayView<uint64> Entities, TConstArrayView<uint64> ArchetypeIDs) = 0;
+	// Notification of moved entities to the given ArchetypeID
+	virtual void BulkMoveEntity(double Time, TConstArrayView<uint64> Entities, TConstArrayView<uint64> ArchetypeIDs) = 0;
+	// Notification of destroyed entities
+	virtual void BulkDestroyEntity(double Time, TConstArrayView<uint64> Entities) = 0;
+
+	/*
+	* Append a new begin event of a region from the trace session.
+	*
+	* @param Name		The string name of the region.
+	* @param Time		The time in seconds of the begin event of this region.
+	*/
+	virtual void AppendRegionBegin(const TCHAR* Name, double Time) = 0;
+
+	/*
+	* Append a new begin event of a region from the trace session.
+	*
+	* @param Name		The string name of the region.
+	* @param Name		The ID of the region. Used to uniquely identify regions with the same name
+	* @param Time		The time in seconds of the begin event of this region.
+	*/
+	virtual void AppendRegionBegin(const TCHAR* Name, uint64 ID, double Time) = 0;
+	
+	/*
+	* Append a new end event of a region from the trace session (by Name).
+	*
+	* @param Name		The string name of the region.
+	* @param Time		The time in seconds of the end event of this region.
+	*/
+	virtual void AppendRegionEnd(const TCHAR* Name, double Time) = 0;
+
+	/*
+	* Append a new end event of a region from the trace session (by ID).
+	*
+	* @param ID			The ID of the region.
+	* @param Time		The time in seconds of the end event of this region.
+	*/
+	virtual void AppendRegionEnd(const uint64 ID, double Time) = 0;
+	
+	/**
+	 * Called from the analyzer once all events have been processed.
+	 * Allows postprocessing and error reporting for regions that were never closed.
+	 */
+	virtual void OnAnalysisSessionEnded() = 0;
+};
+
+////////////////////////////////////////////////////////////////////////////////////////////////////
+
+MASSINSIGHTSANALYSIS_API FName GetMassInsightsProviderName();
+MASSINSIGHTSANALYSIS_API const IMassInsightsProvider& ReadMassInsightsProvider(const TraceServices::IAnalysisSession& Session);
+MASSINSIGHTSANALYSIS_API IEditableMassInsightsProvider& EditMassInsightsProvider(TraceServices::IAnalysisSession& Session);
+	
+} // namespace MassInsightsAnalysis
diff --git a/Engine/Plugins/MassInsights/Source/MassInsightsUI/MassInsightsUI.Build.cs b/Engine/Plugins/MassInsights/Source/MassInsightsUI/MassInsightsUI.Build.cs
new file mode 100644
--- /dev/null
+++ b/Engine/Plugins/MassInsights/Source/MassInsightsUI/MassInsightsUI.Build.cs
@@ -0,0 +1,23 @@
+namespace UnrealBuildTool.Rules
+{
+	public class MassInsightsUI : ModuleRules
+	{
+		public MassInsightsUI(ReadOnlyTargetRules Target) : base(Target)
+		{
+			PrivateDependencyModuleNames.AddRange(
+			new string[]
+			{
+				"Core",
+				"CoreUObject",
+				"SlateCore",
+				"Slate",
+				"InputCore",
+				"TraceInsights",
+				"TraceInsightsCore",
+				"TraceServices",
+				"MassInsightsAnalysis",
+			});
+		}
+	}
+}
+
diff --git a/Engine/Plugins/MassInsights/Source/MassInsightsUI/Private/MassInsightsTimingTrack.cpp b/Engine/Plugins/MassInsights/Source/MassInsightsUI/Private/MassInsightsTimingTrack.cpp
new file mode 100644
--- /dev/null
+++ b/Engine/Plugins/MassInsights/Source/MassInsightsUI/Private/MassInsightsTimingTrack.cpp
@@ -0,0 +1,322 @@
+﻿// Copyright Epic Games, Inc. All Rights Reserved.
+#include "MassInsightsTimingTrack.h"
+
+#include "Framework/MultiBox/MultiBoxBuilder.h"
+
+// TraceServices
+#include "Common/ProviderLock.h"
+
+// TraceInsights
+#include "MassInsightsUIModule.h"
+#include "Framework/Commands/Commands.h"
+#include "InsightsCore/Common/TimeUtils.h"
+#include "Insights/ITimingViewSession.h"
+#include "Insights/IUnrealInsightsModule.h"
+#include "Insights/ViewModels/TimingEventSearch.h"
+#include "Insights/ViewModels/TimingTrackViewport.h"
+#include "Insights/ViewModels/TooltipDrawState.h"
+#include "Insights/ViewModels/TimingEvent.h"
+#include "MassInsightsUI/Widgets/SMassInsightsAnalysisTab.h"
+#include "Modules/ModuleManager.h"
+
+#define LOCTEXT_NAMESPACE "RegionsTimingTrack"
+
+class IUnrealInsightsModule;
+class FTimingEvent;
+
+namespace MassInsightsUI
+{
+
+////////////////////////////////////////////////////////////////////////////////////////////////////
+// FFileActivityTimingViewCommands
+////////////////////////////////////////////////////////////////////////////////////////////////////
+
+FMassInsightsViewCommands::FMassInsightsViewCommands()
+: ::TCommands<FMassInsightsViewCommands>(
+	TEXT("FMassInsightsViewCommands"),
+	NSLOCTEXT("Contexts", "FMassInsightsViewCommands", "Insights - Timing View - Mass Processor"),
+	NAME_None,
+	TEXT("InsightsStyle"))
+{
+}
+
+////////////////////////////////////////////////////////////////////////////////////////////////////
+
+FMassInsightsViewCommands::~FMassInsightsViewCommands()
+{
+}
+
+////////////////////////////////////////////////////////////////////////////////////////////////////
+	
+UE_DISABLE_OPTIMIZATION_SHIP
+void FMassInsightsViewCommands::RegisterCommands()
+{
+	UI_COMMAND(ShowHideRegionTrack,
+		"Mass Processor Phase Tracks",
+		"Shows/hides the Tracks demarcating the begin and end of MassProcessor phases.",
+		EUserInterfaceActionType::ToggleButton,
+		FInputChord(EKeys::Y));
+}
+UE_ENABLE_OPTIMIZATION_SHIP
+
+////////////////////////////////////////////////////////////////////////////////////////////////////
+
+void FMassInsightsSharedState::OnBeginSession(UE::Insights::Timing::ITimingViewSession& InSession)
+{
+	BindCommands();
+	MassInsightsTrack.Reset();
+}
+
+////////////////////////////////////////////////////////////////////////////////////////////////////
+
+void FMassInsightsSharedState::OnEndSession(UE::Insights::Timing::ITimingViewSession& InSession)
+{
+	MassInsightsTrack.Reset();
+}
+
+////////////////////////////////////////////////////////////////////////////////////////////////////
+
+void FMassInsightsSharedState::Tick(UE::Insights::Timing::ITimingViewSession& InSession,
+	const TraceServices::IAnalysisSession& InAnalysisSession)
+{
+	if (!MassInsightsTrack.IsValid())
+	{
+		MassInsightsTrack = MakeShared<FMassInsightsTrack>(*this);
+		MassInsightsTrack->SetOrder(FTimingTrackOrder::First);
+		MassInsightsTrack->SetVisibilityFlag(true);
+		InSession.AddScrollableTrack(MassInsightsTrack);
+	}
+
+	if (TSharedPtr<MassInsights::SMassInsightsAnalysisTab> Analysis = FMassInsightsUIModule::Get().GetAnalysisTab())
+	{
+		Analysis->SetSession(&InSession, &InAnalysisSession);
+	}
+}
+
+////////////////////////////////////////////////////////////////////////////////////////////////////
+
+void FMassInsightsSharedState::ShowHideRegionsTrack()
+{
+	bShowHideRegionsTrack = !bShowHideRegionsTrack;
+
+	if (MassInsightsTrack.IsValid())
+	{
+		MassInsightsTrack->SetVisibilityFlag(bShowHideRegionsTrack);
+	}
+
+	if (bShowHideRegionsTrack)
+	{
+		MassInsightsTrack->SetDirtyFlag();
+	}
+}
+
+////////////////////////////////////////////////////////////////////////////////////////////////////
+
+void FMassInsightsSharedState::ExtendOtherTracksFilterMenu(UE::Insights::Timing::ITimingViewSession& InSession,
+                                                            FMenuBuilder& InOutMenuBuilder)
+{
+	InOutMenuBuilder.BeginSection("Mass Processor Regions", LOCTEXT("ContextMenu_Section_Regions", "Mass Processor Regions"));
+	InOutMenuBuilder.AddMenuEntry(FMassInsightsViewCommands::Get().ShowHideRegionTrack);
+	InOutMenuBuilder.EndSection();
+}
+
+////////////////////////////////////////////////////////////////////////////////////////////////////
+
+void FMassInsightsSharedState::BindCommands()
+{
+	FMassInsightsViewCommands::Register();	
+}
+
+//////////////////////////////////////////////////////////////////////////
+
+INSIGHTS_IMPLEMENT_RTTI(FMassInsightsTrack)
+
+FMassInsightsTrack::FMassInsightsTrack(FMassInsightsSharedState& InSharedState)
+	 : FTimingEventsTrack(TEXT("Mass Phases")), SharedState(InSharedState)
+{
+	IUnrealInsightsModule& UnrealInsightsModule = FModuleManager::LoadModuleChecked<IUnrealInsightsModule>("TraceInsights");
+	AnalysisSession = UnrealInsightsModule.GetAnalysisSession();
+}
+
+FMassInsightsTrack::~FMassInsightsTrack()
+{
+}
+
+void FMassInsightsTrack::BuildContextMenu(FMenuBuilder& MenuBuilder)
+{
+	FTimingEventsTrack::BuildContextMenu(MenuBuilder);
+}
+
+//////////////////////////////////////////////////////////////////////////
+
+void FMassInsightsTrack::InitTooltip(FTooltipDrawState& InOutTooltip, const ITimingEvent& InTooltipEvent) const
+{
+	if (InTooltipEvent.CheckTrack(this) && InTooltipEvent.Is<FTimingEvent>())
+	{
+		const FTimingEvent& TooltipEvent = InTooltipEvent.As<FTimingEvent>();
+
+		auto MatchEvent = [this, &TooltipEvent](double InStartTime, double InEndTime, uint32 InDepth)
+		{
+			return InDepth == TooltipEvent.GetDepth()
+				&& InStartTime == TooltipEvent.GetStartTime()
+				&& InEndTime == TooltipEvent.GetEndTime();
+		};
+
+		FTimingEventSearchParameters SearchParameters(TooltipEvent.GetStartTime(), TooltipEvent.GetEndTime(), ETimingEventSearchFlags::StopAtFirstMatch, MatchEvent);
+		FindRegionEvent(SearchParameters, [this, &InOutTooltip, &TooltipEvent] (
+			double InFoundStartTime,
+			double InFoundEndTime,
+			uint32 InFoundDepth,
+			const MassInsightsAnalysis::FMassInsights& InRegion)
+		{
+			InOutTooltip.Reset();
+			InOutTooltip.AddTitle(InRegion.Text, FLinearColor::White);
+			InOutTooltip.AddNameValueTextLine(TEXT("Duration:"),  UE::Insights::FormatTimeAuto(InRegion.EndTime-InRegion.BeginTime));
+			InOutTooltip.AddNameValueTextLine(TEXT("Depth:"),  FString::FromInt(InRegion.Depth));
+			InOutTooltip.UpdateLayout();
+		});
+	}
+}
+
+void FMassInsightsTrack::BuildDrawState(ITimingEventsTrackDrawStateBuilder& Builder,
+	const ITimingTrackUpdateContext& Context)
+{
+	const FTimingTrackViewport& Viewport = Context.GetViewport();
+
+	const MassInsightsAnalysis::IMassInsightsProvider& RegionProvider = MassInsightsAnalysis::ReadMassInsightsProvider(*AnalysisSession);
+	TraceServices::FProviderReadScopeLock RegionProviderScopedLock(RegionProvider);
+
+	// whe're counting only non-empty lanes, so we can collapse empty ones in the visualization.
+	int32 CurDepth = 0;
+	RegionProvider.EnumerateLanes([this, Viewport, &CurDepth, &Builder](const MassInsightsAnalysis::FMassInsightsLane& Lane, const int32 Depth)
+	{
+		bool RegionHadEvents = false;
+		Lane.EnumerateRegions(Viewport.GetStartTime(), Viewport.GetEndTime(), [&Builder, &RegionHadEvents, &CurDepth](const MassInsightsAnalysis::FMassInsights& Region) -> bool
+		{
+			RegionHadEvents = true;
+			Builder.AddEvent(Region.BeginTime, Region.EndTime,CurDepth, Region.Text);
+			return true;
+		});
+
+		if (RegionHadEvents) CurDepth++;
+	});
+}
+
+//////////////////////////////////////////////////////////////////////////
+
+void FMassInsightsTrack::BuildFilteredDrawState(ITimingEventsTrackDrawStateBuilder& Builder, const ITimingTrackUpdateContext& Context)
+{
+	const TSharedPtr<ITimingEventFilter> EventFilterPtr = Context.GetEventFilter();
+	if (EventFilterPtr.IsValid() && EventFilterPtr->FilterTrack(*this))
+	{
+		bool bFilterOnlyByEventType = false; // this is the most often use case, so the below code tries to optimize it
+		TCHAR* FilterEventType = 0;
+		if (EventFilterPtr->Is<FTimingEventFilterByEventType>())
+		{
+			bFilterOnlyByEventType = true;
+			const FTimingEventFilterByEventType& EventFilter = EventFilterPtr->As<FTimingEventFilterByEventType>();
+			FilterEventType = reinterpret_cast<TCHAR*>(EventFilter.GetEventType());
+		}
+
+		if (AnalysisSession.IsValid())
+		{
+			const MassInsightsAnalysis::IMassInsightsProvider& RegionProvider = MassInsightsAnalysis::ReadMassInsightsProvider(*AnalysisSession);
+			TraceServices::FProviderReadScopeLock RegionProviderScopedLock(RegionProvider);
+
+			const FTimingTrackViewport& Viewport = Context.GetViewport();
+
+			if (bFilterOnlyByEventType)
+			{
+				int32 CurDepth = 0;
+				RegionProvider.EnumerateLanes([this, Viewport, &CurDepth, &Builder, FilterEventType](const MassInsightsAnalysis::FMassInsightsLane& Lane, const int32 Depth)
+					{
+						bool RegionHadEvents = false;
+						Lane.EnumerateRegions(Viewport.GetStartTime(), Viewport.GetEndTime(), [&Builder, &RegionHadEvents, &CurDepth, FilterEventType](const MassInsightsAnalysis::FMassInsights& Region) -> bool
+							{
+								RegionHadEvents = true;
+								if (Region.Text == FilterEventType)
+								{
+									Builder.AddEvent(Region.BeginTime, Region.EndTime, CurDepth, Region.Text);
+								}
+								return true;
+							});
+
+						if (RegionHadEvents) CurDepth++;
+					});
+			}
+			else // generic filter
+			{
+				//TODO: if (EventFilterPtr->FilterEvent(TimingEvent))
+			}
+		}
+	}
+
+}
+
+//////////////////////////////////////////////////////////////////////////
+
+const TSharedPtr<const ITimingEvent> FMassInsightsTrack::SearchEvent(
+	const FTimingEventSearchParameters& InSearchParameters) const
+{
+	TSharedPtr<const ITimingEvent> FoundEvent;
+
+	FindRegionEvent(InSearchParameters, [this, &FoundEvent](double InFoundStartTime, double InFoundEndTime, uint32 InFoundDepth, const MassInsightsAnalysis::FMassInsights& InEvent)
+	{
+		FoundEvent = MakeShared<FTimingEvent>(SharedThis(this), InFoundStartTime, InFoundEndTime, InFoundDepth, reinterpret_cast<uint64>(InEvent.Text));
+	});
+
+	return FoundEvent;
+}
+
+//////////////////////////////////////////////////////////////////////////
+
+bool FMassInsightsTrack::FindRegionEvent(const FTimingEventSearchParameters& InParameters,
+	TFunctionRef<void(double, double, uint32, const MassInsightsAnalysis::FMassInsights&)> InFoundPredicate) const
+{
+	// If the query start time is larger than the end of the session return false.
+	{
+		TraceServices::FAnalysisSessionReadScope SessionReadScope(*AnalysisSession.Get());
+		if (AnalysisSession.IsValid() && InParameters.StartTime > AnalysisSession->GetDurationSeconds())
+		{
+			return false;
+		}
+	}
+
+	return TTimingEventSearch<MassInsightsAnalysis::FMassInsights>::Search(
+		InParameters,
+		[this](TTimingEventSearch<MassInsightsAnalysis::FMassInsights>::FContext& InContext)
+		{
+			const MassInsightsAnalysis::IMassInsightsProvider& RegionProvider = MassInsightsAnalysis::ReadMassInsightsProvider(*AnalysisSession);
+			TraceServices::FProviderReadScopeLock RegionProviderScopedLock(RegionProvider);
+
+			RegionProvider.EnumerateRegions(InContext.GetParameters().StartTime, InContext.GetParameters().EndTime, [&InContext](const MassInsightsAnalysis::FMassInsights& Region)
+			{
+				InContext.Check(Region.BeginTime, Region.EndTime, Region.Depth, Region);
+
+				if (!InContext.ShouldContinueSearching())
+				{
+					return false;
+				}
+
+				return true;
+			});
+		},
+		TTimingEventSearch<MassInsightsAnalysis::FMassInsights>::NoMatch);
+}
+
+//////////////////////////////////////////////////////////////////////////
+
+void FMassInsightsTrack::SetFilterConfigurator(TSharedPtr<UE::Insights::FFilterConfigurator> InFilterConfigurator)
+{
+	if (FilterConfigurator != InFilterConfigurator)
+	{
+		FilterConfigurator = InFilterConfigurator;
+		SetDirtyFlag();
+	}
+}
+
+//////////////////////////////////////////////////////////////////////////
+
+} // namespace Insights
+
+#undef LOCTEXT_NAMESPACE
diff --git a/Engine/Plugins/MassInsights/Source/MassInsightsUI/Private/MassInsightsTimingTrack.h b/Engine/Plugins/MassInsights/Source/MassInsightsUI/Private/MassInsightsTimingTrack.h
new file mode 100644
--- /dev/null
+++ b/Engine/Plugins/MassInsights/Source/MassInsightsUI/Private/MassInsightsTimingTrack.h
@@ -0,0 +1,79 @@
+﻿// Copyright Epic Games, Inc. All Rights Reserved.
+#pragma once
+
+#include "Framework/Commands/Commands.h"
+#include "Insights/ITimingViewExtender.h"
+#include "Insights/ViewModels/TimingEventsTrack.h"
+#include "MassInsightsAnalysis/Model/MassInsights.h"
+#include "Templates/SharedPointer.h"
+
+////////////////////////////////////////////////////////////////////////////////////////////////////
+
+namespace MassInsightsUI
+{
+
+class FMassInsightsTrack;
+
+class FMassInsightsViewCommands : public TCommands<FMassInsightsViewCommands>
+{
+public:
+	FMassInsightsViewCommands();
+	virtual ~FMassInsightsViewCommands() override;
+	virtual void RegisterCommands() override;
+
+public:
+	TSharedPtr<FUICommandInfo> ShowHideRegionTrack;
+};
+
+class FMassInsightsSharedState : public UE::Insights::Timing::ITimingViewExtender, public TSharedFromThis<FMassInsightsSharedState>
+{
+	friend class FMassInsightsTrack;
+	
+public:
+	virtual ~FMassInsightsSharedState() override = default;
+
+	//////////////////////////////////////////////////
+	// ITimingViewExtender interface
+
+	virtual void OnBeginSession(UE::Insights::Timing::ITimingViewSession& InSession) override;
+	virtual void OnEndSession(UE::Insights::Timing::ITimingViewSession& InSession) override;
+	virtual void Tick(UE::Insights::Timing::ITimingViewSession& InSession, const TraceServices::IAnalysisSession& InAnalysisSession) override;
+	virtual void ExtendOtherTracksFilterMenu(UE::Insights::Timing::ITimingViewSession& InSession, FMenuBuilder& InOutMenuBuilder) override;
+
+	//////////////////////////////////////////////////
+	void ShowHideRegionsTrack();
+	bool IsRegionsTrackVisible() const {return bShowHideRegionsTrack;};
+	
+	void BindCommands();
+
+private:
+	TSharedPtr<FMassInsightsTrack> MassInsightsTrack;
+
+	bool bShowHideRegionsTrack = true;
+};
+
+class FMassInsightsTrack : public FTimingEventsTrack
+{
+	INSIGHTS_DECLARE_RTTI(FMassInsightsTrack, FTimingEventsTrack)
+
+public:
+	FMassInsightsTrack(FMassInsightsSharedState& InSharedState);
+	virtual ~FMassInsightsTrack() override;
+
+	virtual void BuildContextMenu(FMenuBuilder& MenuBuilder) override;
+	virtual void InitTooltip(FTooltipDrawState& InOutTooltip, const ITimingEvent& InTooltipEvent) const override;
+	virtual const TSharedPtr<const ITimingEvent> SearchEvent(const FTimingEventSearchParameters& InSearchParameters) const override;
+	virtual void BuildDrawState(ITimingEventsTrackDrawStateBuilder& Builder, const ITimingTrackUpdateContext& Context) override;
+	virtual void BuildFilteredDrawState(ITimingEventsTrackDrawStateBuilder& Builder, const ITimingTrackUpdateContext& Context) override;
+	virtual void SetFilterConfigurator(TSharedPtr<UE::Insights::FFilterConfigurator> InFilterConfigurator) override;
+
+protected:
+	bool FindRegionEvent(const FTimingEventSearchParameters& InParameters, TFunctionRef<void(double, double, uint32, const MassInsightsAnalysis::FMassInsights&)> InFoundPredicate) const;
+
+private:
+	TSharedPtr<UE::Insights::FFilterConfigurator> FilterConfigurator;
+	FMassInsightsSharedState& SharedState;
+	TSharedPtr<const TraceServices::IAnalysisSession> AnalysisSession;
+};
+
+} // namespace MassInsightsUI
\ No newline at end of file
diff --git a/Engine/Plugins/MassInsights/Source/MassInsightsUI/Private/MassInsightsUI/MassInsightsStyle.h b/Engine/Plugins/MassInsights/Source/MassInsightsUI/Private/MassInsightsUI/MassInsightsStyle.h
new file mode 100644
--- /dev/null
+++ b/Engine/Plugins/MassInsights/Source/MassInsightsUI/Private/MassInsightsUI/MassInsightsStyle.h
@@ -0,0 +1,49 @@
+﻿// Copyright Epic Games, Inc. All Rights Reserved.
+
+#pragma once
+
+#include "Brushes/SlateImageBrush.h"
+#include "Styling/CoreStyle.h"
+#include "Math/ColorList.h"
+#include "Styling/SlateStyle.h"
+#include "Misc/Paths.h"
+#include "Styling/SlateStyleRegistry.h"
+#include "Styling/SlateTypes.h"
+#include "Styling/SlateStyleMacros.h"
+
+
+namespace MassInsights
+{
+
+class FMassInsightsStyle final : public FSlateStyleSet
+{
+public:
+	FMassInsightsStyle()
+		: FSlateStyleSet("MassInsightsStyle")
+	{
+		const FVector2D Icon16x16(16.0f, 16.0f);
+		SetContentRoot(FPaths::EngineContentDir() / TEXT("Editor/Slate"));
+		SetCoreContentRoot(FPaths::EngineContentDir() / TEXT("Slate"));
+
+		Set("MassProfiler.Icon.Small", new CORE_IMAGE_BRUSH("Icons/Profiler/profiler_stats_40x", FVector2D(16.0f, 16.0f)));
+
+		FSlateStyleRegistry::RegisterSlateStyle(*this);
+	}
+
+	static FMassInsightsStyle& Get()
+	{
+		static FMassInsightsStyle Inst;
+		return Inst;
+	}
+	
+	~FMassInsightsStyle()
+	{
+		FSlateStyleRegistry::UnRegisterSlateStyle(*this);
+	}
+};
+
+#undef IMAGE_BRUSH
+
+} //namespace SlateInsights
+
+
diff --git a/Engine/Plugins/MassInsights/Source/MassInsightsUI/Private/MassInsightsUI/Widgets/Common.cpp b/Engine/Plugins/MassInsights/Source/MassInsightsUI/Private/MassInsightsUI/Widgets/Common.cpp
new file mode 100644
--- /dev/null
+++ b/Engine/Plugins/MassInsights/Source/MassInsightsUI/Private/MassInsightsUI/Widgets/Common.cpp
@@ -0,0 +1,18 @@
+﻿// Copyright Epic Games, Inc. All Rights Reserved.
+
+#include "Common.h"
+
+#include "InsightsCore/Common/TimeUtils.h"
+
+namespace MassInsightsUI
+{
+	FText FTableCellFormatterTimeHMS::FormatValue(const TOptional<UE::Insights::FTableCellValue>& InValue) const
+	{
+		if (InValue.IsSet())
+		{
+			const double Value = InValue.GetValue().Double;
+			return FText::FromString(UE::Insights::FormatTimeHMS(Value, UE::Insights::FTimeValue::Microsecond));
+		}
+		return FText();
+	}
+}
diff --git a/Engine/Plugins/MassInsights/Source/MassInsightsUI/Private/MassInsightsUI/Widgets/Common.h b/Engine/Plugins/MassInsights/Source/MassInsightsUI/Private/MassInsightsUI/Widgets/Common.h
new file mode 100644
--- /dev/null
+++ b/Engine/Plugins/MassInsights/Source/MassInsightsUI/Private/MassInsightsUI/Widgets/Common.h
@@ -0,0 +1,20 @@
+﻿// Copyright Epic Games, Inc. All Rights Reserved.
+
+#pragma once
+
+#include "Delegates/Delegate.h"
+#include "InsightsCore/Table/ViewModels/TableCellValueFormatter.h"
+
+namespace MassInsightsUI
+{
+	DECLARE_DELEGATE_OneParam(FOnSelectedArchetype, uint64);
+
+	/**
+	 * Formats the value of a table's cell into time HHh MMm SS.SSSSSSs
+	 * ie. 2h 10m 32.893027s
+	 */
+	class FTableCellFormatterTimeHMS : public UE::Insights::FTableCellValueFormatter
+	{
+		virtual FText FormatValue(const TOptional<UE::Insights::FTableCellValue>& InValue) const override;
+	};
+}
diff --git a/Engine/Plugins/MassInsights/Source/MassInsightsUI/Private/MassInsightsUI/Widgets/SArchetypeDetails.cpp b/Engine/Plugins/MassInsights/Source/MassInsightsUI/Private/MassInsightsUI/Widgets/SArchetypeDetails.cpp
new file mode 100644
--- /dev/null
+++ b/Engine/Plugins/MassInsights/Source/MassInsightsUI/Private/MassInsightsUI/Widgets/SArchetypeDetails.cpp
@@ -0,0 +1,273 @@
+﻿// Copyright Epic Games, Inc. All Rights Reserved.
+
+#include "SArchetypeDetails.h"
+
+#include "Common/ProviderLock.h"
+#include "Hash/BuzHash.h"
+#include "Insights/IInsightsManager.h"
+#include "Insights/IUnrealInsightsModule.h"
+#include "InsightsCore/Table/ViewModels/Table.h"
+#include "InsightsCore/Table/ViewModels/TableColumn.h"
+#include "MassInsightsAnalysis/Model/MassInsights.h"
+#include "Modules/ModuleManager.h"
+#include "Widgets/Input/SButton.h"
+
+#define LOCTEXT_NAMESPACE "MassInsightsUI::SArchetypeDetails"
+
+namespace MassInsightsUI
+{
+	FName NAME_COLUMN_FragmentName = TEXT("FragmentName");
+	FName NAME_COLUMN_FragmentType = TEXT("FragmentType");
+	
+	class SFragmentListRow : public SMultiColumnTableRow< TSharedPtr<SArchetypeDetails::FFragmentListEntry> >
+	{
+		using Super = SMultiColumnTableRow< TSharedPtr<SArchetypeDetails::FFragmentListEntry> >;
+	public:
+		SLATE_BEGIN_ARGS(SFragmentListRow) {}
+			SLATE_ARGUMENT(const MassInsightsAnalysis::FMassFragmentInfo*, FragmentInfo)
+		SLATE_END_ARGS()
+
+		void Construct(const FArguments& InArgs, const TSharedRef<STableViewBase>& InOwnerTableView);
+		TSharedRef<SWidget> GenerateWidgetForColumn(const FName& InColumnName) override;
+	private:
+		const MassInsightsAnalysis::FMassFragmentInfo* FragmentInfo = nullptr;
+	};
+
+	void SFragmentListRow::Construct(const FArguments& InArgs, const TSharedRef<STableViewBase>& InOwnerTableView)
+	{
+		FragmentInfo = InArgs._FragmentInfo;
+		SetEnabled(true);
+		SMultiColumnTableRow< TSharedPtr<SArchetypeDetails::FFragmentListEntry> >::Construct(Super::FArguments(), InOwnerTableView);
+	}
+
+	TSharedRef<SWidget> SFragmentListRow::GenerateWidgetForColumn(const FName& InColumnName)
+	{
+		if (InColumnName == NAME_COLUMN_FragmentName)
+		{
+			return SNew( STextBlock )
+					.Text(FText::FromString(FragmentInfo->GetName()));
+		}
+		else if (InColumnName == NAME_COLUMN_FragmentType)
+		{
+			// Not a real checkbox, but using it 
+			return SNew(SButton)
+					.Text_Lambda([this]()
+					{
+						switch (FragmentInfo->Type)
+						{
+						case MassInsightsAnalysis::EFragmentType::Fragment: return LOCTEXT("Fragment", "Fragment");
+						case MassInsightsAnalysis::EFragmentType::Tag: return LOCTEXT("Tag", "Tag");
+						case MassInsightsAnalysis::EFragmentType::Shared: return LOCTEXT("Shared", "Shared");
+						case MassInsightsAnalysis::EFragmentType::Unknown: // Fallthrough
+						default:
+							return LOCTEXT("Unknown", "Unknown");
+						}
+					})
+				.ForegroundColor_Lambda(
+				[this]()
+				{
+					switch (FragmentInfo->Type)
+					{
+						case MassInsightsAnalysis::EFragmentType::Fragment: return FColor(10, 120, 180);
+						case MassInsightsAnalysis::EFragmentType::Tag: // Fallthrough
+						case MassInsightsAnalysis::EFragmentType::Shared: // Fallthrough
+						case MassInsightsAnalysis::EFragmentType::Unknown: // Fallthrough
+						default:
+							return FColor(90, 120, 10);
+					}
+				});
+		}
+		else
+		{
+			return SNew( STextBlock )
+					.Text(FText::FromString(TEXT("-")));
+		}
+	}
+	
+	SArchetypeDetails::SArchetypeDetails()
+		: Table(MakeShared<UE::Insights::FTable>())
+	{
+	}
+
+	void SArchetypeDetails::Construct(const FArguments& InArgs)
+	{
+		IUnrealInsightsModule& UnrealInsightsModule = FModuleManager::LoadModuleChecked<IUnrealInsightsModule>("TraceInsights");
+		TSharedPtr<::Insights::IInsightsManager> InsightsManager = UnrealInsightsModule.GetInsightsManager();
+		TSharedPtr<const TraceServices::IAnalysisSession> CurrentSession = UnrealInsightsModule.GetAnalysisSession();
+		Session = CurrentSession;
+		
+		ChildSlot
+		[
+			SAssignNew(ListView, STreeView<TSharedPtr<FFragmentListEntry>>)
+			.SelectionMode(ESelectionMode::Single)
+			.TreeItemsSource(&FilteredSortedFragmentList)
+			.OnGetChildren(this, &SArchetypeDetails::TreeView_OnGetChildren)
+			.OnGenerateRow(this, &SArchetypeDetails::TreeView_OnGenerateRow)
+			.HeaderRow
+			(
+				SAssignNew(HeaderRow, SHeaderRow)
+				.Visibility(EVisibility::Visible)
+			)
+		];
+
+		TArray<TSharedRef<UE::Insights::FTableColumn>> Columns;
+		{
+			using namespace UE::Insights;
+			TSharedRef<UE::Insights::FTableColumn> ColumnRef = MakeShared<FTableColumn>(NAME_COLUMN_FragmentName);
+			FTableColumn& Column = *ColumnRef;
+			Column.SetShortName(LOCTEXT("FragmentName", "Name"));
+			Column.SetTitleName(LOCTEXT("FragmentName", "Name"));
+			Column.SetDescription(LOCTEXT("FragmentName", "Name"));
+			Column.SetHorizontalAlignment(HAlign_Left);
+			Column.SetInitialWidth(206.0f);
+			Column.SetMinWidth(42.0f);
+			Column.SetFlags(ETableColumnFlags::ShouldBeVisible);
+			Column.SetDataType(ETableCellDataType::Text);
+			Columns.Add(ColumnRef);
+		}
+		{
+			using namespace UE::Insights;
+			TSharedRef<UE::Insights::FTableColumn> ColumnRef = MakeShared<FTableColumn>(NAME_COLUMN_FragmentType);
+			FTableColumn& Column = *ColumnRef;
+			Column.SetShortName(LOCTEXT("FragmentName", "Type"));
+			Column.SetTitleName(LOCTEXT("FragmentName", "Type"));
+			Column.SetDescription(LOCTEXT("FragmentName", "Type"));
+			Column.SetHorizontalAlignment(HAlign_Left);
+			Column.SetInitialWidth(206.0f);
+			Column.SetMinWidth(42.0f);
+			Column.SetFlags(ETableColumnFlags::ShouldBeVisible);
+			Column.SetDataType(ETableCellDataType::Text);
+			Columns.Add(ColumnRef);
+		}
+
+		Table->SetColumns(Columns);
+		
+		for (const TSharedRef<UE::Insights::FTableColumn>& ColumnRef : Table->GetColumns())
+		{
+			if (ColumnRef->ShouldBeVisible())
+			{
+				ShowColumn(ColumnRef->GetId());
+			}
+		}
+	}
+
+	void SArchetypeDetails::SetArchetype(uint64 InArchetypeID)
+	{
+		if (ArchetypeID != InArchetypeID)
+		{
+			ArchetypeID = InArchetypeID;
+			bInvalidArchetype = false;
+			bArchetypeDirty = true;
+		}
+	}
+
+	TSharedRef<ITableRow> SArchetypeDetails::TreeView_OnGenerateRow(TSharedPtr<FFragmentListEntry> NodePtr, const TSharedRef<STableViewBase>& OwnerTable)
+	{
+		TSharedRef<ITableRow> TableRow = SNew(SFragmentListRow, OwnerTable).FragmentInfo(NodePtr->Fragment);
+		return TableRow;
+	}
+
+	void SArchetypeDetails::TreeView_OnGetChildren(TSharedPtr<FFragmentListEntry> InParent, TArray<TSharedPtr<FFragmentListEntry>>& OutChildren)
+	{
+	}
+
+	void SArchetypeDetails::Tick(const FGeometry& AllottedGeometry, const double InCurrentTime, const float InDeltaTime)
+	{
+		SCompoundWidget::Tick(AllottedGeometry, InCurrentTime, InDeltaTime);
+
+		if (bArchetypeDirty)
+		{
+			bArchetypeDirty = false;
+			bSortingDirty = true;
+			FragmentListData.Reset();
+			
+
+			const MassInsightsAnalysis::IMassInsightsProvider& Provider = MassInsightsAnalysis::ReadMassInsightsProvider(*Session);
+			{
+				TraceServices::FProviderReadScopeLock ProviderReadScopeLock(Provider);
+
+				const MassInsightsAnalysis::FMassArchetypeInfo* ArchetypeInfo = Provider.FindArchetypeById(ArchetypeID);
+				if (ArchetypeInfo != nullptr)
+				{
+					for (int32 Index = 0; Index < ArchetypeInfo->Fragments.Num(); Index++)
+					{
+						FragmentListData.Emplace(MakeShared<FFragmentListEntry>(
+							FFragmentListEntry
+							{
+								.Fragment = ArchetypeInfo->Fragments[Index]
+							}));
+					}
+				}
+			}
+		}
+
+		if (bSortingDirty)
+		{
+			bSortingDirty = false;
+			FilteredSortedFragmentList.Reset();
+			
+			// Sort by Type then Name
+			FilteredSortedFragmentList.Append(FragmentListData);
+			Algo::Sort(FilteredSortedFragmentList, [](const TSharedPtr<FFragmentListEntry>& Lhs, const TSharedPtr<FFragmentListEntry>& Rhs)
+			{
+				if (Lhs->Fragment->Type == Rhs->Fragment->Type)
+				{
+					return Lhs->Fragment->Name < Rhs->Fragment->Name;
+				}
+				return Lhs->Fragment->Type < Rhs->Fragment->Type;
+			});
+
+			ListView.Pin()->RebuildList();
+			
+		}
+	}
+
+	void SArchetypeDetails::ShowColumn(const FName ColumnId)
+	{
+		UE::Insights::FTableColumn& Column = *Table->FindColumnChecked(ColumnId);
+		Column.Show();
+
+		SHeaderRow::FColumn::FArguments ColumnArgs;
+		ColumnArgs
+			.ColumnId(Column.GetId())
+			.DefaultLabel(Column.GetShortName())
+			.HAlignHeader(Column.GetHorizontalAlignment())
+			.VAlignHeader(VAlign_Center)
+			.HAlignCell(HAlign_Fill)
+			.VAlignCell(VAlign_Fill)
+			.InitialSortMode(Column.GetInitialSortMode())
+			.FillWidth(Column.GetInitialWidth())
+			.HeaderContent()
+			[
+				SNew(SBox)
+				.HeightOverride(24.0f)
+				.Padding(FMargin(0.0f))
+				.VAlign(VAlign_Center)
+				[
+					SNew(STextBlock)
+					.Text_Lambda([this, ColumnID = Column.GetId()]()->FText
+					{
+						const UE::Insights::FTableColumn& Column = *Table->FindColumnChecked(ColumnID);
+						return Column.GetShortName();
+					})
+				]
+			];
+
+		int32 ColumnIndex = 0;
+		const int32 NewColumnPosition = Table->GetColumnPositionIndex(ColumnId);
+		const int32 NumColumns = HeaderRow->GetColumns().Num();
+		for (; ColumnIndex < NumColumns; ColumnIndex++)
+		{
+			const SHeaderRow::FColumn& CurrentColumn = HeaderRow->GetColumns()[ColumnIndex];
+			const int32 CurrentColumnPosition = Table->GetColumnPositionIndex(CurrentColumn.ColumnId);
+			if (NewColumnPosition < CurrentColumnPosition)
+			{
+				break;
+			}
+		}
+
+		HeaderRow->InsertColumn(ColumnArgs, ColumnIndex);
+	}
+}
+
+#undef LOCTEXT_NAMESPACE
\ No newline at end of file
diff --git a/Engine/Plugins/MassInsights/Source/MassInsightsUI/Private/MassInsightsUI/Widgets/SArchetypeDetails.h b/Engine/Plugins/MassInsights/Source/MassInsightsUI/Private/MassInsightsUI/Widgets/SArchetypeDetails.h
new file mode 100644
--- /dev/null
+++ b/Engine/Plugins/MassInsights/Source/MassInsightsUI/Private/MassInsightsUI/Widgets/SArchetypeDetails.h
@@ -0,0 +1,59 @@
+﻿// Copyright Epic Games, Inc. All Rights Reserved.
+
+#pragma once
+#include "MassInsightsAnalysis/Model/MassInsights.h"
+#include "Widgets/SCompoundWidget.h"
+#include "Widgets/Views/STreeView.h"
+
+namespace UE::Insights
+{
+	class FTable;
+}
+
+namespace TraceServices
+{
+	class IAnalysisSession;
+}
+
+namespace MassInsightsAnalysis
+{
+	struct FMassFragmentInfo;
+}
+
+namespace MassInsightsUI
+{
+	class SArchetypeDetails : public SCompoundWidget
+	{
+	public:
+		SArchetypeDetails();
+		SLATE_BEGIN_ARGS(SArchetypeDetails) {}
+			
+		SLATE_END_ARGS()
+
+		void Construct(const FArguments& _args);
+		void SetArchetype(uint64 ArchetypeID);
+	private:
+		friend class SFragmentListRow;
+		struct FFragmentListEntry
+		{
+			const MassInsightsAnalysis::FMassFragmentInfo* Fragment = nullptr;
+		};
+
+		TSharedRef<ITableRow> TreeView_OnGenerateRow(TSharedPtr<FFragmentListEntry> NodePtr, const TSharedRef<STableViewBase>& OwnerTable);
+		void TreeView_OnGetChildren(TSharedPtr<FFragmentListEntry> InParent, TArray<TSharedPtr<FFragmentListEntry>>& OutChildren);
+		
+		virtual void Tick(const FGeometry& AllottedGeometry, const double InCurrentTime, const float InDeltaTime) override;
+		void ShowColumn(const FName ColumnId);
+
+		TSharedPtr<UE::Insights::FTable> Table;
+		uint64 ArchetypeID = 0;
+		bool bArchetypeDirty = true;
+		bool bSortingDirty = true;
+		bool bInvalidArchetype = true;
+		TSharedPtr<const TraceServices::IAnalysisSession> Session;
+		TWeakPtr<STreeView<TSharedPtr<FFragmentListEntry>>> ListView;
+		TSharedPtr<SHeaderRow> HeaderRow;
+		TArray<TSharedPtr<FFragmentListEntry>> FragmentListData;
+		TArray<TSharedPtr<FFragmentListEntry>> FilteredSortedFragmentList;
+	};
+}
diff --git a/Engine/Plugins/MassInsights/Source/MassInsightsUI/Private/MassInsightsUI/Widgets/SEntityEventAggregationTableView.cpp b/Engine/Plugins/MassInsights/Source/MassInsightsUI/Private/MassInsightsUI/Widgets/SEntityEventAggregationTableView.cpp
new file mode 100644
--- /dev/null
+++ b/Engine/Plugins/MassInsights/Source/MassInsightsUI/Private/MassInsightsUI/Widgets/SEntityEventAggregationTableView.cpp
@@ -0,0 +1,597 @@
+﻿// Copyright Epic Games, Inc. All Rights Reserved.
+
+#include "SEntityEventAggregationTableView.h"
+
+#include "Algo/ForEach.h"
+#include "Common/ProviderLock.h"
+#include "Insights/IInsightsManager.h"
+#include "InsightsCore/Table/ViewModels/TableCellValueFormatter.h"
+#include "InsightsCore/Table/ViewModels/TableCellValueGetter.h"
+#include "InsightsCore/Table/ViewModels/TableCellValueSorter.h"
+#include "InsightsCore/Table/ViewModels/TableColumn.h"
+#include "Modules/ModuleManager.h"
+#include "Widgets/Input/SButton.h"
+#include "Widgets/Text/STextBlock.h"
+
+#define LOCTEXT_NAMESPACE "MassInsights::SEntityTimelineTableView"
+
+namespace MassInsights
+{
+	static FName COLUMN_EntityID = FName("EntityID");
+	static FName COLUMN_Archetype = FName("Archetype");
+	static FName COLUMN_CreateTime = FName("CreateTime");
+	static FName COLUMN_EventCount = FName("EventCount");
+	static FName COLUMN_LastEventTime = FName("LastEventTime");
+	static FName COLUMN_Alive = FName("Alive");
+	
+	INSIGHTS_IMPLEMENT_RTTI(FEntityEventAggregateRecordHandle)
+	
+	FEntityEventAggregateRecordHandle::FEntityEventAggregateRecordHandle(int32 InRecordIndex)
+		: FBaseTreeNode(NAME_None, false /* bInIsGroup */)
+		, RecordIndex(InRecordIndex)
+	{
+	}
+
+	int32 FEntityEventAggregateRecordHandle::GetRecordIndex() const
+	{
+		return RecordIndex;
+	}
+
+	// Helper template to read the value referenced by a data member pointer from FEntityEventAggregationNode and convert it to
+	// a FTableCellValue
+	template<typename MemberPtr>
+	class SEntityEventAggregationTableView::TAggregationRecordMemberAccessor : public UE::Insights::FTableCellValueGetter
+	{
+	public:
+		using FTableCellValue = UE::Insights::FTableCellValue;
+		using FTableColumn = UE::Insights::FTableColumn;
+		using FBaseTreeNode = UE::Insights::FBaseTreeNode;
+		
+		explicit TAggregationRecordMemberAccessor(MemberPtr InPtr, const TArray<FEntityEventAggregationNode>& InEventAggregationRecords)
+			: Ptr(InPtr)
+			, EventAggregationRecords(InEventAggregationRecords)
+		{}
+				
+		virtual const TOptional<FTableCellValue> GetValue(const FTableColumn& Column, const FBaseTreeNode& Node) const override
+		{
+			const int32 Index = static_cast<const FEntityEventAggregateRecordHandle&>(Node).GetRecordIndex();
+			const SEntityEventAggregationTableView::FEntityEventAggregationNode& AggregationRecord = EventAggregationRecords[Index];
+			return FTableCellValue(AggregationRecord.*Ptr);
+		}
+		MemberPtr Ptr;
+		const TArray<FEntityEventAggregationNode>& EventAggregationRecords;
+	};
+	
+	template <typename MemberPtr>
+	constexpr auto SEntityEventAggregationTableView::MakeAccessor(MemberPtr Ptr)
+	{
+		return MakeShareable(new TAggregationRecordMemberAccessor<MemberPtr>(Ptr, EventAggregationRecords));
+	}
+
+	void SEntityEventAggregationTableRow::Construct(const FArguments& InArgs, const TSharedRef<STableViewBase>& InOwnerTableView)
+	{
+		TablePtr = InArgs._TablePtr;
+		RecordHandle = InArgs._RowHandle;
+		AggregationTableView = InArgs._ViewModel;
+		OnArchetypeSelected = InArgs._OnArchetypeSelected;
+		SetEnabled(true);
+		Super::Construct(Super::FArguments(), InOwnerTableView);
+	}
+
+	TSharedRef<SWidget> SEntityEventAggregationTableRow::GenerateWidgetForColumn(const FName& InColumnName)
+	{
+		int32 Index = RecordHandle->GetRecordIndex();
+		TSharedPtr<SEntityEventAggregationTableView> VM = AggregationTableView.Pin();
+		const SEntityEventAggregationTableView::FEntityEventAggregationNode& AggregationRecord = VM->EventAggregationRecords[Index];
+
+		TSharedPtr<UE::Insights::FTableColumn> ColumnPtr = TablePtr->FindColumnChecked(InColumnName);
+		if (!ColumnPtr.IsValid())
+		{
+			return SNew( STextBlock )
+				.Text(LOCTEXT("Unknown Column Message", "Unknown Column"));
+		}
+		
+		if (InColumnName == COLUMN_EntityID)
+		{
+			return SNew( STextBlock )
+				.Text(FText::Format(LOCTEXT("EntityID", "{0}"), AggregationRecord.EntityID ));
+		}
+		else if (InColumnName == COLUMN_Archetype)
+		{
+			return SNew(SButton)
+				.Text(LOCTEXT("Archetype", "Archetype"))
+				.OnReleased(this, &SEntityEventAggregationTableRow::OnReleasedArchetypeButton, AggregationRecord.LastArchetype);
+		}
+		else if (InColumnName == COLUMN_CreateTime)
+		{
+			return SNew( STextBlock )
+				.Text(ColumnPtr->GetValueAsText(*RecordHandle));
+		}
+		else if (InColumnName == COLUMN_EventCount)
+		{
+			return SNew( STextBlock )
+				.Text(FText::Format(LOCTEXT("EventCount", "{0}"), AggregationRecord.Events ));
+		}
+		else if (InColumnName == COLUMN_LastEventTime)
+		{
+			return SNew( STextBlock )
+				.Text(ColumnPtr->GetValueAsText(*RecordHandle));
+		}
+		else if (InColumnName == COLUMN_Alive)
+		{
+			if (AggregationRecord.bDestroyed)
+			{
+				return SNew(SHorizontalBox)
+				+SHorizontalBox::Slot()
+				.HAlign(HAlign_Center)
+				[
+					SNew(SImage)
+					.Image(FAppStyle::GetBrush("Icons.X"))
+					.DesiredSizeOverride(FVector2D(16.0, 16.0))
+					.ColorAndOpacity(FSlateColor(FColor(160.0f, 20.0f, 30.0f)))
+				];
+			}
+			else
+			{
+				return SNullWidget::NullWidget;
+			}
+		}
+		else
+		{
+			return SNew( STextBlock )
+				.Text(LOCTEXT("Unknown Column", "Unknown Column"));
+		}
+	}
+
+	void SEntityEventAggregationTableRow::OnReleasedArchetypeButton(uint64 ArchetypeID)
+	{
+		OnArchetypeSelected.ExecuteIfBound(ArchetypeID);
+	}
+
+
+	SEntityEventAggregationTableView::SEntityEventAggregationTableView()
+		: Table(MakeShared<UE::Insights::FTable>())
+	{
+	}
+
+	SEntityEventAggregationTableView::~SEntityEventAggregationTableView()
+	{
+	}
+
+	void SEntityEventAggregationTableView::Construct(const FArguments& InArgs)
+	{
+		OnArchetypeSelected = InArgs._OnArchetypeSelected;
+		OnRowSelected = InArgs._OnRowSelected;
+		
+		ChildSlot
+		[
+			SNew(SVerticalBox)
+
+			+ SVerticalBox::Slot()
+			.AutoHeight()
+			.Padding(2.0f, 2.0f, 2.0f, 2.0f)
+			[
+				// TODO: This is a good place to add filtering tools
+				SNullWidget::NullWidget
+			]
+			
+			+SVerticalBox::Slot()
+			.VAlign(VAlign_Fill)
+			.FillHeight(1.0f)
+			.Padding(2.0f, 2.0f, 2.0f, 2.0f)
+			[
+				SNew(SHorizontalBox)
+				
+				+ SHorizontalBox::Slot()
+				.FillWidth(1.0f)
+				.Padding(0.0f)
+				[
+					SAssignNew(TreeView, STreeView<TSharedPtr<FEntityEventAggregateRecordHandle>>)
+					.SelectionMode(ESelectionMode::Single)
+					.TreeItemsSource(&FilteredTreeViewRows)
+					.OnGetChildren(this, &SEntityEventAggregationTableView::TreeView_OnGetChildren)
+					.OnGenerateRow(this, &SEntityEventAggregationTableView::TreeView_OnGenerateRow)
+					.OnSelectionChanged_Lambda([this](TSharedPtr<FEntityEventAggregateRecordHandle> Item, ESelectInfo::Type SelectInfo)
+					{
+						TArray<TSharedPtr<FEntityEventAggregateRecordHandle>> SelectedItems = TreeView->GetSelectedItems();
+						if (SelectedItems.Num() == 0)
+						{
+							FEntityEventSummaryRowSelectedParams Params
+							{
+								.IsSelected = false
+							};
+							this->OnRowSelected.ExecuteIfBound(Params);
+						}
+						else if (EventAggregationRecords.IsValidIndex(SelectedItems[0]->GetRecordIndex()))
+						{
+							const FEntityEventAggregationNode& EventAggregationNode = EventAggregationRecords[SelectedItems[0]->GetRecordIndex()];
+							FEntityEventSummaryRowSelectedParams Params
+							{
+								.IsSelected = true,
+								.EntityID = EventAggregationNode.EntityID,
+								.FirstEventTime = EventAggregationNode.FirstEventTime,
+								.LastEventTime = EventAggregationNode.LastEventTime,
+								.TotalEvents = EventAggregationNode.Events
+							};
+							this->OnRowSelected.ExecuteIfBound(Params);
+						}
+					})
+					
+					.HeaderRow
+					(
+						SAssignNew(HeaderRow, SHeaderRow)
+						.Visibility(EVisibility::Visible)
+					)
+				]
+			]
+
+			// Status bar
+			+ SVerticalBox::Slot()
+			.AutoHeight()
+			.HAlign(HAlign_Fill)
+			.VAlign(VAlign_Bottom)
+			.Padding(0.0f)
+			[
+				SNew(SBorder)
+				.BorderImage(FAppStyle::Get().GetBrush("WhiteBrush"))
+				.BorderBackgroundColor(FLinearColor(0.05f, 0.1f, 0.2f, 1.0f))
+				.HAlign(HAlign_Center)
+				[
+					SNew(STextBlock)
+					.Margin(FMargin(4.0f, 1.0f, 4.0f, 1.0f))
+					.Text_Raw(this, &SEntityEventAggregationTableView::GetStatusBarText)
+					.ColorAndOpacity(FLinearColor(1.0f, 0.75f, 0.5f, 1.0f))
+					.Visibility(EVisibility::Visible)
+				]
+			]
+		];
+
+		IUnrealInsightsModule& UnrealInsightsModule = FModuleManager::LoadModuleChecked<IUnrealInsightsModule>("TraceInsights");
+		TSharedPtr<::Insights::IInsightsManager> InsightsManager = UnrealInsightsModule.GetInsightsManager();
+		if (InsightsManager.IsValid())
+		{
+			InsightsManager->GetSessionChangedEvent().AddSP(this, &SEntityEventAggregationTableView::InsightsManager_OnSessionChanged);
+			InsightsManager->GetSessionAnalysisCompletedEvent().AddSP(this, &SEntityEventAggregationTableView::InsightsManager_OnSessionAnalysisComplete);
+		}
+		
+		TSharedPtr<const TraceServices::IAnalysisSession> CurrentSession = UnrealInsightsModule.GetAnalysisSession();
+		Session = CurrentSession;
+
+		TArray<TSharedRef<UE::Insights::FTableColumn>> Columns;		
+		
+		{
+			using namespace UE::Insights;
+			TSharedRef<FTableColumn> ColumnRef = MakeShared<FTableColumn>(COLUMN_CreateTime);
+			FTableColumn& Column = *ColumnRef;
+			Column.SetShortName(LOCTEXT("First Event Time", "First Event Time"));
+			Column.SetTitleName(LOCTEXT("First Event Time", "First Event Time"));
+			Column.SetDescription(LOCTEXT("First Event Time", "First Event Time"));
+			Column.SetHorizontalAlignment(HAlign_Left);
+			Column.SetInitialWidth(100.0f);
+			Column.SetMinWidth(42.0f);
+			Column.SetFlags(ETableColumnFlags::ShouldBeVisible);
+			Column.SetDataType(ETableCellDataType::Double);
+			
+			TSharedRef<ITableCellValueGetter> Getter = MakeAccessor(&FEntityEventAggregationNode::FirstEventTime);//MakeShared<FDisplayNameValueGetter>();
+			Column.SetValueGetter(Getter);
+			
+			Column.SetValueFormatter(MakeShared<MassInsightsUI::FTableCellFormatterTimeHMS>());
+
+			TSharedRef<ITableCellValueSorter> Sorter = MakeShared<FSorterByDoubleValue>(ColumnRef);
+			Column.SetValueSorter(Sorter);
+			
+			Columns.Add(ColumnRef);
+		}
+		
+		{
+			using namespace UE::Insights;
+			TSharedRef<FTableColumn> ColumnRef = MakeShared<FTableColumn>(COLUMN_EntityID);
+			FTableColumn& Column = *ColumnRef;
+			Column.SetShortName(LOCTEXT("EntityID", "EntityID"));
+			Column.SetTitleName(LOCTEXT("EntityID", "EntityID"));
+			Column.SetDescription(LOCTEXT("EntityID", "EntityID"));
+			Column.SetHorizontalAlignment(HAlign_Left);
+			Column.SetInitialWidth(100.0f);
+			Column.SetMinWidth(42.0f);
+			Column.SetFlags(ETableColumnFlags::ShouldBeVisible);
+			Column.SetDataType(ETableCellDataType::Custom);
+			
+			Columns.Add(ColumnRef);
+		}
+		
+		{
+			using namespace UE::Insights;
+			TSharedRef<FTableColumn> ColumnRef = MakeShared<FTableColumn>(COLUMN_EventCount);
+			FTableColumn& Column = *ColumnRef;
+			Column.SetShortName(LOCTEXT("Total Events", "Total Events"));
+			Column.SetTitleName(LOCTEXT("Total Events", "Total Events"));
+			Column.SetDescription(LOCTEXT("Total Events", "Total Events"));
+			Column.SetHorizontalAlignment(HAlign_Left);
+			Column.SetInitialWidth(80.0f);
+			Column.SetMinWidth(41.0f);
+			Column.SetFlags(ETableColumnFlags::ShouldBeVisible);
+			Column.SetDataType(ETableCellDataType::Int64);
+			
+			Columns.Add(ColumnRef);
+		}
+		
+		{
+			using namespace UE::Insights;
+			TSharedRef<FTableColumn> ColumnRef = MakeShared<FTableColumn>(COLUMN_Alive);
+			FTableColumn& Column = *ColumnRef;
+			Column.SetShortName(LOCTEXT("Status", "Status"));
+			Column.SetTitleName(LOCTEXT("Status", "Status"));
+			Column.SetDescription(LOCTEXT("Status", "Status"));
+			Column.SetHorizontalAlignment(HAlign_Left);
+			Column.SetInitialWidth(30.0f);
+			Column.SetMinWidth(30.0f);
+			Column.SetFlags(ETableColumnFlags::ShouldBeVisible);
+			Column.SetDataType(ETableCellDataType::Int64);
+			
+			Columns.Add(ColumnRef);
+		}
+		
+		{
+			using namespace UE::Insights;
+			TSharedRef<FTableColumn> ColumnRef = MakeShared<FTableColumn>(COLUMN_LastEventTime);
+			FTableColumn& Column = *ColumnRef;
+			Column.SetShortName(LOCTEXT("Last Event Time", "Last Event Time"));
+			Column.SetTitleName(LOCTEXT("Last Event Time", "Last Event Time"));
+			Column.SetDescription(LOCTEXT("Last Event Time", "Last Event Time"));
+			Column.SetHorizontalAlignment(HAlign_Left);
+			Column.SetInitialWidth(100.0f);
+			Column.SetMinWidth(42.0f);
+			Column.SetFlags(ETableColumnFlags::ShouldBeVisible);
+			Column.SetDataType(ETableCellDataType::Custom);
+			
+			TSharedRef<ITableCellValueGetter> Getter = MakeAccessor(&FEntityEventAggregationNode::LastEventTime);//MakeShared<FDisplayNameValueGetter>();
+			Column.SetValueGetter(Getter);
+			
+			TSharedRef<ITableCellValueFormatter> Formatter = MakeShared<MassInsightsUI::FTableCellFormatterTimeHMS>();
+			Column.SetValueFormatter(Formatter);
+			
+			Columns.Add(ColumnRef);
+		}
+		
+		Table->SetColumns(Columns);
+
+		for (const TSharedRef<UE::Insights::FTableColumn>& ColumnRef : Table->GetColumns())
+		{
+			if (ColumnRef->ShouldBeVisible())
+			{
+				ShowColumn(ColumnRef->GetId());
+			}
+		}
+	}
+
+	void SEntityEventAggregationTableView::Tick(const FGeometry& AllottedGeometry, const double InCurrentTime, const float InDeltaTime)
+	{
+		SCompoundWidget::Tick(AllottedGeometry, InCurrentTime, InDeltaTime);
+		
+		if (bRebuildTree)
+		{
+			EntityToRecordIndexMap.Empty();
+			EventsProcessed = 0;
+			EventAggregationRecords.Empty();
+			TreeView->RebuildList();
+			
+			bRebuildTree = false;
+		}
+
+		IncrementalUpdate();
+	}
+
+	void SEntityEventAggregationTableView::InsightsManager_OnSessionChanged()
+	{
+		IUnrealInsightsModule& UnrealInsightsModule = FModuleManager::LoadModuleChecked<IUnrealInsightsModule>("TraceInsights");
+
+		TSharedPtr<const TraceServices::IAnalysisSession> CurrentSession = UnrealInsightsModule.GetAnalysisSession();
+		
+		Session = CurrentSession;
+		Reset();
+	}
+
+	void SEntityEventAggregationTableView::InsightsManager_OnSessionAnalysisComplete()
+	{
+	}
+
+	TSharedRef<ITableRow> SEntityEventAggregationTableView::TreeView_OnGenerateRow(TSharedPtr<FEntityEventAggregateRecordHandle> RowHandle,
+	                                                                 const TSharedRef<STableViewBase>& OwnerTable)
+	{
+		TSharedPtr<SWidget> SharedThisAsWidget = AsShared();
+		auto SharedThis = StaticCastSharedPtr<SEntityEventAggregationTableView>(SharedThisAsWidget);
+		TSharedRef<ITableRow> TableRow =
+			SNew(MassInsights::SEntityEventAggregationTableRow, OwnerTable)
+			.TablePtr(Table)
+			.RowHandle(RowHandle)
+			.ViewModel(SharedThis)
+			.OnArchetypeSelected_Lambda([this](uint64 ArchetypeID)
+			{
+				OnArchetypeSelected.ExecuteIfBound(ArchetypeID);
+			});
+	
+		return TableRow;
+	}
+
+	void SEntityEventAggregationTableView::TreeView_OnGetChildren(TSharedPtr<FEntityEventAggregateRecordHandle> InParent,
+		TArray<TSharedPtr<FEntityEventAggregateRecordHandle>>& OutChildren)
+	{
+		return;
+	}
+
+	void SEntityEventAggregationTableView::Reset()
+	{
+		bRebuildTree = true;
+	}
+
+	void SEntityEventAggregationTableView::IncrementalUpdate()
+	{
+		if (Session.IsValid())
+		{
+			uint64 PostUpdate_NextIncrementalEventIndex = EventsProcessed;
+
+			const int32 PreUpdateRecordCount = EventAggregationRecords.Num();
+
+			// @TODO: Would be better if this were async instead of on Tick
+			const MassInsightsAnalysis::IMassInsightsProvider& Provider = MassInsightsAnalysis::ReadMassInsightsProvider(*Session);
+			{
+				TraceServices::FProviderReadScopeLock ProviderReadScopeLock(Provider);
+
+				constexpr uint64 MaximumEventsToCheck = 100000;
+
+				EstimatedEventCount = Provider.GetEntityEventCount();
+				
+				Provider.EnumerateEntityEvents(
+					EventsProcessed,
+					MaximumEventsToCheck,
+					[this, &PostUpdate_NextIncrementalEventIndex](const MassInsightsAnalysis::FMassEntityEventRecord& Event, uint64 EventIndex)
+					{
+						++PostUpdate_NextIncrementalEventIndex;
+						
+						switch (Event.Operation)
+						{
+						case MassInsightsAnalysis::EMassEntityEventType::Created:
+							{
+								const int32 RecordIndex = EventAggregationRecords.Num();
+								EntityToRecordIndexMap.Add(Event.Entity, EventAggregationRecords.Num());
+								EventAggregationRecords.Emplace(FEntityEventAggregationNode
+									{
+										.EntityID = Event.Entity,
+										.LastArchetype = Event.ArchetypeID,
+										.Events = 1,
+										.bDestroyed = false,
+										.FirstEventTime = Event.Time,
+										.LastEventTime = Event.Time,
+										.RowHandle = MakeShared<FEntityEventAggregateRecordHandle>(RecordIndex)
+									});
+							}
+							break;
+						case MassInsightsAnalysis::EMassEntityEventType::ArchetypeChange:
+							{
+								int32* IndexPtr = EntityToRecordIndexMap.Find(Event.Entity);
+								if (IndexPtr != nullptr)
+								{
+									// Update the row, maintain the creation time		
+									FEntityEventAggregationNode& RecordIndex = EventAggregationRecords[*IndexPtr];
+									RecordIndex.LastArchetype = Event.ArchetypeID;
+									RecordIndex.Events++;
+									RecordIndex.LastEventTime = Event.Time;
+								}
+								else
+								{
+									const int32 RecordIndex = EventAggregationRecords.Num();
+									EntityToRecordIndexMap.Add(Event.Entity, EventAggregationRecords.Num());
+									EventAggregationRecords.Emplace(FEntityEventAggregationNode
+									{
+										.EntityID = Event.Entity,
+										.LastArchetype = Event.ArchetypeID,
+										.Events = 1,
+										.bDestroyed = false,
+										.FirstEventTime = Event.Time,
+										.LastEventTime = Event.Time,
+										.RowHandle = MakeShared<FEntityEventAggregateRecordHandle>(RecordIndex)
+									});
+								}							
+							}
+							break;
+						case MassInsightsAnalysis::EMassEntityEventType::Destroyed:
+							{
+								int32* IndexPtr = EntityToRecordIndexMap.Find(Event.Entity);
+								if (IndexPtr != nullptr)
+								{
+									FEntityEventAggregationNode& Record = EventAggregationRecords[*IndexPtr];
+									Record.Events++;
+									Record.bDestroyed = true;
+									Record.LastEventTime = Event.Time;
+								}
+								else
+								{
+									const int32 RecordIndex = EventAggregationRecords.Num();
+									EntityToRecordIndexMap.Add(Event.Entity, EventAggregationRecords.Num());
+									EventAggregationRecords.Emplace(FEntityEventAggregationNode
+									{
+										.EntityID = Event.Entity,
+										.LastArchetype = Event.ArchetypeID,
+										.Events = 1,
+										.bDestroyed = true,
+										.FirstEventTime = Event.Time,
+										.LastEventTime = Event.Time,
+										.RowHandle = MakeShared<FEntityEventAggregateRecordHandle>(RecordIndex)
+									});
+								}
+							}
+							break;
+						}
+					});
+			}
+			
+			if (EventsProcessed != PostUpdate_NextIncrementalEventIndex)
+			{
+				for (int32 RecordIndex = PreUpdateRecordCount; RecordIndex < EventAggregationRecords.Num(); RecordIndex++)
+				{
+					FilteredTreeViewRows.Add(EventAggregationRecords[RecordIndex].RowHandle);
+				}
+				
+				TreeView->RequestListRefresh();
+			}
+			EventsProcessed = PostUpdate_NextIncrementalEventIndex;
+		};
+	}
+
+	void SEntityEventAggregationTableView::ShowColumn(const FName ColumnId)
+	{
+		UE::Insights::FTableColumn& Column = *Table->FindColumnChecked(ColumnId);
+		Column.Show();
+
+		SHeaderRow::FColumn::FArguments ColumnArgs;
+		ColumnArgs
+			.ColumnId(Column.GetId())
+			.DefaultLabel(Column.GetShortName())
+			.HAlignHeader(Column.GetHorizontalAlignment())
+			.VAlignHeader(VAlign_Center)
+			.HAlignCell(HAlign_Fill)
+			.VAlignCell(VAlign_Fill)
+			.InitialSortMode(Column.GetInitialSortMode())
+			.FillWidth(Column.GetInitialWidth())
+			.HeaderContent()
+			[
+				SNew(SBox)
+				.HeightOverride(24.0f)
+				.Padding(FMargin(0.0f))
+				.VAlign(VAlign_Center)
+				[
+					SNew(STextBlock)
+					.Text(this, &SEntityEventAggregationTableView::GetColumnHeaderText, Column.GetId())
+				]
+			];
+
+		int32 ColumnIndex = 0;
+		const int32 NewColumnPosition = Table->GetColumnPositionIndex(ColumnId);
+		const int32 NumColumns = HeaderRow->GetColumns().Num();
+		for (; ColumnIndex < NumColumns; ColumnIndex++)
+		{
+			const SHeaderRow::FColumn& CurrentColumn = HeaderRow->GetColumns()[ColumnIndex];
+			const int32 CurrentColumnPosition = Table->GetColumnPositionIndex(CurrentColumn.ColumnId);
+			if (NewColumnPosition < CurrentColumnPosition)
+			{
+				break;
+			}
+		}
+
+		HeaderRow->InsertColumn(ColumnArgs, ColumnIndex);
+	}
+
+	FText SEntityEventAggregationTableView::GetColumnHeaderText(const FName ColumnId) const
+	{
+		const UE::Insights::FTableColumn& Column = *Table->FindColumnChecked(ColumnId);
+		return Column.GetShortName();
+	}
+
+	FText SEntityEventAggregationTableView::GetStatusBarText() const
+	{
+		return FText::Format(LOCTEXT("", "-- Entities:{0} Events:{1}/{2}--"), FilteredTreeViewRows.Num(), EventsProcessed, EstimatedEventCount);
+	}
+
+}
+
+#undef LOCTEXT_NAMESPACE
\ No newline at end of file
diff --git a/Engine/Plugins/MassInsights/Source/MassInsightsUI/Private/MassInsightsUI/Widgets/SEntityEventAggregationTableView.h b/Engine/Plugins/MassInsights/Source/MassInsightsUI/Private/MassInsightsUI/Widgets/SEntityEventAggregationTableView.h
new file mode 100644
--- /dev/null
+++ b/Engine/Plugins/MassInsights/Source/MassInsightsUI/Private/MassInsightsUI/Widgets/SEntityEventAggregationTableView.h
@@ -0,0 +1,139 @@
+﻿// Copyright Epic Games, Inc. All Rights Reserved.
+
+#pragma once
+#include "SFragmentTableRow.h"
+#include "Framework/Commands/UICommandList.h"
+#include "InsightsCore/Table/ViewModels/BaseTreeNode.h"
+#include "Widgets/SCompoundWidget.h"
+#include "Widgets/Views/STreeView.h"
+#include "Widgets/Views/STableRow.h"
+#include "Widgets/Views/SListView.h"
+#include "InsightsCore/Table/ViewModels/Table.h"
+#include "MassInsightsUI/Widgets/Common.h"
+
+namespace MassInsights
+{
+	struct FEntityEventAggregateRecordHandle : public UE::Insights::FBaseTreeNode
+	{
+		INSIGHTS_DECLARE_RTTI(FEntityEventAggregateRecordHandle, UE::Insights::FBaseTreeNode);
+
+	public:
+		explicit FEntityEventAggregateRecordHandle(int32 InRecordIndex);
+
+		int32 GetRecordIndex() const;
+	private:
+		
+		// Index used to reference event aggregation data
+		int32 RecordIndex;
+	};
+
+	class SEntityEventAggregationTableView;
+	
+	class SEntityEventAggregationTableRow : public SMultiColumnTableRow< TSharedPtr<FEntityEventAggregateRecordHandle> >
+	{
+		using Super = SMultiColumnTableRow< TSharedPtr<FEntityEventAggregateRecordHandle> >;
+	public:
+		SLATE_BEGIN_ARGS(SEntityEventAggregationTableRow) {}
+			SLATE_ARGUMENT(TSharedPtr<UE::Insights::FTable>, TablePtr)
+			SLATE_ARGUMENT(TSharedPtr<FEntityEventAggregateRecordHandle>, RowHandle)
+			SLATE_ARGUMENT(TSharedPtr<SEntityEventAggregationTableView>, ViewModel)
+			SLATE_EVENT(MassInsightsUI::FOnSelectedArchetype, OnArchetypeSelected)
+		SLATE_END_ARGS()
+
+		void Construct(const FArguments& InArgs, const TSharedRef<STableViewBase>& InOwnerTableView);
+
+		virtual TSharedRef<SWidget> GenerateWidgetForColumn(const FName& InColumnName) override;
+
+	private:
+		void OnReleasedArchetypeButton(uint64 ArchetypeID);
+		TSharedPtr<UE::Insights::FTable> TablePtr;
+		TWeakPtr<SEntityEventAggregationTableView> AggregationTableView;
+		TSharedPtr<FEntityEventAggregateRecordHandle> RecordHandle;
+		MassInsightsUI::FOnSelectedArchetype OnArchetypeSelected;
+		
+	};
+
+	// Struct to make it easier to extend passing parameters to this delegate
+	struct FEntityEventSummaryRowSelectedParams
+	{
+		// False if deselected, all fields are undefined on deselection
+		// True if selected, all fields should describe the row selected
+		bool IsSelected;
+		uint64 EntityID;
+		double FirstEventTime;
+		double LastEventTime;
+		int64 TotalEvents;
+	};
+	DECLARE_DELEGATE_OneParam(FEntityEventContainerRowSelected, const FEntityEventSummaryRowSelectedParams&);
+	
+	class SEntityEventAggregationTableView : public SCompoundWidget
+	{
+		friend class SEntityEventAggregationTableRow;
+	public:
+		SEntityEventAggregationTableView();
+		virtual ~SEntityEventAggregationTableView();
+
+		SLATE_BEGIN_ARGS(SEntityEventAggregationTableView) {}
+			SLATE_EVENT(MassInsightsUI::FOnSelectedArchetype, OnArchetypeSelected)
+			SLATE_EVENT(FEntityEventContainerRowSelected, OnRowSelected)
+		SLATE_END_ARGS()
+
+		void Construct(const FArguments& args);
+		void Reset();
+	private:
+		void InsightsManager_OnSessionChanged();
+		void InsightsManager_OnSessionAnalysisComplete();
+		
+		TSharedRef<ITableRow> TreeView_OnGenerateRow(TSharedPtr<FEntityEventAggregateRecordHandle> NodePtr, const TSharedRef<STableViewBase>& OwnerTable);
+		void TreeView_OnGetChildren(TSharedPtr<FEntityEventAggregateRecordHandle> InParent, TArray<TSharedPtr<FEntityEventAggregateRecordHandle>>& OutChildren);
+	
+		virtual void Tick(const FGeometry& AllottedGeometry, const double InCurrentTime, const float InDeltaTime) override;
+		
+		void IncrementalUpdate();
+		void ShowColumn(const FName ColumnId);
+		FText GetColumnHeaderText(const FName ColumnId) const;
+
+		FText GetStatusBarText() const;
+		
+	private:
+		TSharedPtr<UE::Insights::FTable> Table;
+		TSharedPtr<const TraceServices::IAnalysisSession> Session;
+		TSharedPtr<FUICommandList> CommandList;
+		
+		TSharedPtr<STreeView<TSharedPtr<FEntityEventAggregateRecordHandle>>> TreeView;
+		TSharedPtr<SHeaderRow> HeaderRow;
+		
+		// Collated summary data about the events gathered from the application
+		struct FEntityEventAggregationNode
+		{
+			uint64 EntityID;
+			// ID of the last archetype this entity had
+			uint64 LastArchetype;
+			int32 Events = 0;
+			bool bDestroyed = false;
+			double FirstEventTime;
+			double LastEventTime;
+			TSharedPtr<FEntityEventAggregateRecordHandle> RowHandle;
+		};
+
+		template<typename MemberPtr>
+		class TAggregationRecordMemberAccessor;
+
+		template<typename MemberPtr>
+		constexpr auto MakeAccessor(MemberPtr Ptr);
+		
+		TArray<FEntityEventAggregationNode> EventAggregationRecords;
+
+		TMap<uint64, int32> EntityToRecordIndexMap;
+
+		uint64 EventsProcessed = 0;
+		uint64 EstimatedEventCount = 0;
+
+		bool bRebuildTree = true;
+		
+		TArray<TSharedPtr<FEntityEventAggregateRecordHandle>> FilteredTreeViewRows;
+		
+		MassInsightsUI::FOnSelectedArchetype OnArchetypeSelected;
+		FEntityEventContainerRowSelected OnRowSelected;
+	};
+}
diff --git a/Engine/Plugins/MassInsights/Source/MassInsightsUI/Private/MassInsightsUI/Widgets/SEntityEventsTableView.cpp b/Engine/Plugins/MassInsights/Source/MassInsightsUI/Private/MassInsightsUI/Widgets/SEntityEventsTableView.cpp
new file mode 100644
--- /dev/null
+++ b/Engine/Plugins/MassInsights/Source/MassInsightsUI/Private/MassInsightsUI/Widgets/SEntityEventsTableView.cpp
@@ -0,0 +1,471 @@
+﻿// Copyright Epic Games, Inc. All Rights Reserved.
+
+#include "SEntityEventsTableView.h"
+
+#include "Algo/Find.h"
+#include "Algo/ForEach.h"
+#include "Algo/LevenshteinDistance.h"
+#include "Common/ProviderLock.h"
+#include "Insights/IInsightsManager.h"
+#include "InsightsCore/Common/TimeUtils.h"
+#include "InsightsCore/Table/ViewModels/TableCellValueFormatter.h"
+#include "InsightsCore/Table/ViewModels/TableCellValueGetter.h"
+#include "InsightsCore/Table/ViewModels/TableCellValueSorter.h"
+#include "InsightsCore/Table/ViewModels/TableColumn.h"
+#include "Modules/ModuleManager.h"
+#include "Widgets/Input/SButton.h"
+#include "Widgets/Text/STextBlock.h"
+#include "MassInsightsUI/Widgets/Common.h"
+
+#define LOCTEXT_NAMESPACE "MassInsights::SEntityTimelineTableView"
+
+namespace MassInsightsUI
+{
+	static FName COLUMN_EventTime = TEXT("EventTime");
+	static FName COLUMN_Operation = TEXT("Operation");
+
+	class SEntityEventsTableRow : public SMultiColumnTableRow< TSharedPtr<SEntityEventsTableView::FEventNode> >
+	{
+		using Super = SMultiColumnTableRow< TSharedPtr<SEntityEventsTableView::FEventNode> >;
+	public:
+		SLATE_BEGIN_ARGS(SEntityEventsTableRow) {}
+		SLATE_ARGUMENT(TSharedPtr<UE::Insights::FTable>, TablePtr)
+		SLATE_ARGUMENT(TSharedPtr<SEntityEventsTableView::FEventNode>, EventNode)
+		SLATE_ARGUMENT(TSharedPtr<SEntityEventsTableView>, ViewModel)
+		SLATE_ARGUMENT(TSharedPtr<const TraceServices::IAnalysisSession>, AnalysisSession)
+		SLATE_EVENT(FOnSelectedArchetype, OnArchetypeSelected)
+	SLATE_END_ARGS()
+
+		void Construct(const FArguments& InArgs, const TSharedRef<STableViewBase>& InOwnerTableView);
+
+		virtual TSharedRef<SWidget> GenerateWidgetForColumn(const FName& InColumnName) override;
+
+	private:
+		void OnReleasedArchetypeButton(uint64 ArchetypeID);
+		TSharedPtr<UE::Insights::FTable> TablePtr;
+		TWeakPtr<SEntityEventsTableView> ViewModel;
+		TSharedPtr<const TraceServices::IAnalysisSession> AnalysisSession;
+		TSharedPtr<SEntityEventsTableView::FEventNode> EventNode;
+		FOnSelectedArchetype OnArchetypeSelected;
+		
+	};
+
+	void SEntityEventsTableRow::Construct(const FArguments& InArgs, const TSharedRef<STableViewBase>& InOwnerTableView)
+	{
+		TablePtr = InArgs._TablePtr;
+		EventNode = InArgs._EventNode;
+		ViewModel = InArgs._ViewModel;
+		AnalysisSession = InArgs._AnalysisSession;
+		check(TablePtr);
+		check(EventNode);
+		check(ViewModel.IsValid());
+		check(AnalysisSession);
+		OnArchetypeSelected = InArgs._OnArchetypeSelected;
+		SetEnabled(true);
+		Super::Construct(Super::FArguments(), InOwnerTableView);
+	}
+	
+	TSharedRef<SWidget> SEntityEventsTableRow::GenerateWidgetForColumn(const FName& InColumnName)
+	{
+		const uint64 ProviderEventIndex = EventNode->GetProviderEventIndex();
+		const MassInsightsAnalysis::IMassInsightsProvider& Provider = MassInsightsAnalysis::ReadMassInsightsProvider(*AnalysisSession);
+
+		TValueOrError<MassInsightsAnalysis::FMassEntityEventRecord, void> Event = MakeError();
+		{
+			TraceServices::FProviderReadScopeLock ProviderReadScopeLock(Provider);
+
+			Event = Provider.GetEntityEvent(ProviderEventIndex);
+		}
+		if (!Event.HasValue())
+		{
+			return SNew(STextBlock).Text(LOCTEXT("MissingEventIndex", "N/A'"));
+		}
+		
+		TSharedPtr<UE::Insights::FTableColumn> ColumnPtr = TablePtr->FindColumnChecked(InColumnName);
+
+		if (InColumnName == COLUMN_EventTime)
+		{
+			TOptional<UE::Insights::FTableCellValue> ValueThing = ColumnPtr->GetValue(*EventNode);
+			return SNew( STextBlock )
+				.Text(ColumnPtr->GetValueAsText(*EventNode));
+		}
+		else if (InColumnName == COLUMN_Operation)
+		{
+			const FSlateBrush* Icon = [EventType = Event.GetValue().Operation]
+			{
+				switch (EventType)
+				{
+				case MassInsightsAnalysis::EMassEntityEventType::Created:
+					return FAppStyle::GetBrush("Icons.Plus");
+					break;
+				case MassInsightsAnalysis::EMassEntityEventType::ArchetypeChange:
+					return FAppStyle::GetBrush("Icons.ArrowRight");
+					break;
+				case MassInsightsAnalysis::EMassEntityEventType::Destroyed:
+					return FAppStyle::GetBrush("Icons.X");
+					break;
+				default:
+					return FAppStyle::GetBrush("Icons.Warning");
+				}
+			}();
+			
+			return SNew(SHorizontalBox)
+				+SHorizontalBox::Slot()
+				.HAlign(HAlign_Center)
+				[
+					SNew(SImage)
+					.Image(Icon)
+					.DesiredSizeOverride(FVector2D(16.0, 16.0))
+					.ColorAndOpacity([EventType = Event.GetValue().Operation]()
+					{
+						switch (EventType)
+						{
+						case MassInsightsAnalysis::EMassEntityEventType::Created:
+							return FSlateColor(FColor(30.0f, 200.0f, 20.0f));
+						case MassInsightsAnalysis::EMassEntityEventType::ArchetypeChange:
+							return FSlateColor(FColor(200.0f, 200.0f, 200.0f));
+						case MassInsightsAnalysis::EMassEntityEventType::Destroyed:
+							return FSlateColor(FColor(160.0f, 20.0f, 30.0f));
+						default:
+							return FSlateColor(FColor(255,255,255));
+						}
+					}())
+				];
+		}
+		
+		return SNew( STextBlock )
+				.Text(LOCTEXT("Unknown Column", "Unknown Column"));
+	}
+
+	void SEntityEventsTableRow::OnReleasedArchetypeButton(uint64 ArchetypeID)
+	{
+		OnArchetypeSelected.ExecuteIfBound(ArchetypeID);
+	}
+	
+	SEntityEventsTableView::FEventNode::FEventNode(const FName InName, bool bInIsGroup, uint64 InProviderEntityIndex)
+		: FBaseTreeNode(InName, bInIsGroup)
+		, ProviderEntityIndex(InProviderEntityIndex)
+	{
+	}
+
+	INSIGHTS_IMPLEMENT_RTTI(SEntityEventsTableView::FEventNode);
+
+	uint64 SEntityEventsTableView::FEventNode::GetProviderEventIndex() const
+	{
+		return ProviderEntityIndex;
+	}
+
+	SEntityEventsTableView::SEntityEventsTableView()
+		: Table(MakeShared<UE::Insights::FTable>())
+	{
+	}
+
+	void SEntityEventsTableView::TreeView_OnSelectionChanged(TSharedPtr<FEventNode> EventNode, ESelectInfo::Type Arg) const
+	{
+		if (EventNode)
+		{
+			FOnSelectedEntityEventParams Params;
+			Params.ProviderEventIndex = EventNode->GetProviderEventIndex();
+			
+			OnEntityEventSelected.ExecuteIfBound(Params);
+		}
+	}
+
+	void SEntityEventsTableView::Construct(const FArguments& InArgs)
+	{
+		OnEntityEventSelected = InArgs._OnEntityEventSelected;
+		ChildSlot
+		[
+			SNew(SVerticalBox)
+
+			+ SVerticalBox::Slot()
+			.AutoHeight()
+			.Padding(2.0f, 2.0f, 2.0f, 2.0f)
+			[
+				// TODO Filtering stuff here
+				SNullWidget::NullWidget
+			]
+					
+			+SVerticalBox::Slot()
+			.VAlign(VAlign_Fill)
+			.FillHeight(1.0f)
+			.Padding(2.0f, 2.0f, 2.0f, 2.0f)
+			[
+				SNew(SHorizontalBox)
+						
+				+ SHorizontalBox::Slot()
+				.FillWidth(1.0f)
+				.Padding(0.0f)
+				[
+					SAssignNew(TreeView, STreeView<TSharedPtr<FEventNode>>)
+					.SelectionMode(ESelectionMode::Single)
+					.TreeItemsSource(&FilteredSortedEvents)
+					.OnGetChildren(this, &SEntityEventsTableView::TreeView_OnGetChildren)
+					.OnGenerateRow(this, &SEntityEventsTableView::TreeView_OnGenerateRow)
+					.OnSelectionChanged(this, &SEntityEventsTableView::TreeView_OnSelectionChanged)
+							
+					.HeaderRow
+					(
+						SAssignNew(HeaderRow, SHeaderRow)
+						.Visibility(EVisibility::Visible)
+					)
+				]
+			]
+		];
+
+		IUnrealInsightsModule& UnrealInsightsModule = FModuleManager::LoadModuleChecked<IUnrealInsightsModule>("TraceInsights");
+		TSharedPtr<::Insights::IInsightsManager> InsightsManager = UnrealInsightsModule.GetInsightsManager();
+		if (InsightsManager.IsValid())
+		{
+			InsightsManager->GetSessionChangedEvent().AddSP(this, &SEntityEventsTableView::InsightsManager_OnSessionChanged);
+		}
+		
+		TSharedPtr<const TraceServices::IAnalysisSession> CurrentSession = UnrealInsightsModule.GetAnalysisSession();
+		Session = CurrentSession;
+
+		TArray<TSharedRef<UE::Insights::FTableColumn>> Columns;
+		
+		{
+			using namespace UE::Insights;
+			TSharedRef<FTableColumn> ColumnRef = MakeShared<FTableColumn>(COLUMN_EventTime);
+			FTableColumn& Column = *ColumnRef;
+			Column.SetShortName(LOCTEXT("Event Time", "Event Time"));
+			Column.SetTitleName(LOCTEXT("Event Time", "Event Time"));
+			Column.SetDescription(LOCTEXT("Event Time", "Event Time"));
+			Column.SetHorizontalAlignment(HAlign_Left);
+			Column.SetInitialWidth(206.0f);
+			Column.SetMinWidth(42.0f);
+			Column.SetFlags(ETableColumnFlags::ShouldBeVisible);
+			Column.SetDataType(ETableCellDataType::Double);
+
+			class FEventTimeGetter : public FTableCellValueGetter
+			{
+			public:
+				virtual const TOptional<FTableCellValue> GetValue(const FTableColumn& Column, const FBaseTreeNode& Node) const override
+				{
+					const FEventNode& EventNode = static_cast<const FEventNode&>(Node);
+					const uint64 ProviderEventIndex = EventNode.GetProviderEventIndex();
+
+					IUnrealInsightsModule& UnrealInsightsModule = FModuleManager::LoadModuleChecked<IUnrealInsightsModule>("TraceInsights");
+					TSharedPtr<::Insights::IInsightsManager> InsightsManager = UnrealInsightsModule.GetInsightsManager();
+					if (InsightsManager.IsValid())
+					{
+						TSharedPtr<const TraceServices::IAnalysisSession> CurrentSession = UnrealInsightsModule.GetAnalysisSession();
+						if (CurrentSession.IsValid())
+						{
+							const MassInsightsAnalysis::IMassInsightsProvider& Provider = MassInsightsAnalysis::ReadMassInsightsProvider(*CurrentSession);
+			
+							{
+								TraceServices::FProviderReadScopeLock ProviderReadScopeLock(Provider);
+								TValueOrError<MassInsightsAnalysis::FMassEntityEventRecord, void> EventData = Provider.GetEntityEvent(ProviderEventIndex);
+								if (EventData.HasValue())
+								{
+									return FTableCellValue(EventData.GetValue().Time);
+								}
+							}
+						}
+					}
+					return TOptional<FTableCellValue>();
+				}
+			};
+
+			Column.SetValueGetter(MakeShared<FEventTimeGetter>());
+
+			Column.SetValueFormatter(MakeShared<FTableCellFormatterTimeHMS>());
+			
+			Columns.Add(ColumnRef);
+		}
+		
+		{
+			using namespace UE::Insights;
+			TSharedRef<FTableColumn> ColumnRef = MakeShared<FTableColumn>(COLUMN_Operation);
+			FTableColumn& Column = *ColumnRef;
+			Column.SetShortName(LOCTEXT("Operation", "Operation"));
+			Column.SetTitleName(LOCTEXT("Operation", "Operation"));
+			Column.SetDescription(LOCTEXT("Operation", "Operation"));
+			Column.SetHorizontalAlignment(HAlign_Left);
+			Column.SetInitialWidth(206.0f);
+			Column.SetMinWidth(42.0f);
+			Column.SetFlags(ETableColumnFlags::ShouldBeVisible);
+			Column.SetDataType(ETableCellDataType::Int64);
+			
+			Columns.Add(ColumnRef);
+		}
+		
+		Table->SetColumns(Columns);
+
+		for (const TSharedRef<UE::Insights::FTableColumn>& ColumnRef : Table->GetColumns())
+		{
+			if (ColumnRef->ShouldBeVisible())
+			{
+				ShowColumn(ColumnRef->GetId());
+			}
+		}
+	}
+
+	void SEntityEventsTableView::SetEntities(TConstArrayView<uint64> InEntities)
+	{
+		Entities = InEntities;
+		bRebuildTable = true;
+	}
+
+	TSharedRef<ITableRow> SEntityEventsTableView::TreeView_OnGenerateRow(
+		TSharedPtr<FEventNode> InEventNode,
+		const TSharedRef<STableViewBase>& OwnerTable)
+	{
+		TSharedPtr<SWidget> SharedThisAsWidget = AsShared();
+		auto SharedThis = StaticCastSharedPtr<SEntityEventsTableView>(SharedThisAsWidget);
+		TSharedRef<ITableRow> TableRow =
+			SNew(MassInsightsUI::SEntityEventsTableRow, OwnerTable)
+			.TablePtr(Table)
+			.EventNode(InEventNode)
+			.ViewModel(SharedThis)
+			.AnalysisSession(Session)
+			.OnArchetypeSelected_Lambda([this](uint64 ArchetypeID)
+			{
+				OnArchetypeSelected.ExecuteIfBound(ArchetypeID);
+			});
+	
+		return TableRow;
+	}
+
+	void SEntityEventsTableView::TreeView_OnGetChildren(TSharedPtr<FEventNode> InParent, TArray<TSharedPtr<FEventNode>>& OutChildren)
+	{
+	}
+
+	void SEntityEventsTableView::IncrementalUpdate()
+	{
+		if (Session.IsValid())
+		{
+			const MassInsightsAnalysis::IMassInsightsProvider& Provider = MassInsightsAnalysis::ReadMassInsightsProvider(*Session);
+
+			// Note: Would be better if this were async instead of on Tick
+			{
+				TraceServices::FProviderReadScopeLock ProviderReadScopeLock(Provider);
+
+				constexpr uint64 MaximumEventsToProcess = 10000;
+
+				const int32 PreviousEventNodeCount = Events.Num();
+
+				uint64 NewNextEventIndex = NextEventIndex;
+				Provider.EnumerateEntityEvents(NextEventIndex, MaximumEventsToProcess, [this, &NewNextEventIndex](const MassInsightsAnalysis::FMassEntityEventRecord& Event, uint64 ProviderEventIndex)
+				{
+					++NewNextEventIndex;
+
+					const uint64* EntityPtr = Algo::Find(Entities, Event.Entity);
+					if (!EntityPtr)
+					{
+						return;
+					}
+
+					const int32 EntityIndex = static_cast<int32>(static_cast<uint64>((EntityPtr - Entities.GetData())) / sizeof(uint64));
+					check(EntityIndex >= 0 && EntityIndex < Entities.Num());
+
+					const FName GroupName(TEXT("All"));
+					const bool bIsInGroup = false;
+					TSharedPtr<FEventNode> EventNode = MakeShared<FEventNode>(GroupName, bIsInGroup, ProviderEventIndex);
+					Events.Emplace(MoveTemp(EventNode));
+				});
+				
+				const int32 NewEventNodesAdded = Events.Num() - PreviousEventNodeCount;
+				
+				if (NewEventNodesAdded > 0)
+				{
+					// Should already be sorted
+					FilteredSortedEvents.Reserve(Events.Num());
+					for (int32 Index = PreviousEventNodeCount; Index < Events.Num(); Index++)
+					{
+						FilteredSortedEvents.Add(Events[Index]);
+					}
+					
+					TreeView->RebuildList();
+				}
+				
+				NextEventIndex = NewNextEventIndex;
+				
+			}
+		}
+	}
+
+	void SEntityEventsTableView::Tick(const FGeometry& AllottedGeometry, const double InCurrentTime, const float InDeltaTime)
+	{
+		SCompoundWidget::Tick(AllottedGeometry, InCurrentTime, InDeltaTime);
+
+		if (bRebuildTable)
+		{
+			NextEventIndex = 0;
+			Events.Reset();
+			FilteredSortedEvents.Reset();
+			TreeView->RebuildList();
+			
+			bRebuildTable = false;
+		}
+
+		IncrementalUpdate();
+	}
+
+	void SEntityEventsTableView::ShowColumn(const FName ColumnId)
+	{
+		UE::Insights::FTableColumn& Column = *Table->FindColumnChecked(ColumnId);
+		Column.Show();
+
+		SHeaderRow::FColumn::FArguments ColumnArgs;
+		ColumnArgs
+			.ColumnId(Column.GetId())
+			.DefaultLabel(Column.GetShortName())
+			.HAlignHeader(Column.GetHorizontalAlignment())
+			.VAlignHeader(VAlign_Center)
+			.HAlignCell(HAlign_Fill)
+			.VAlignCell(VAlign_Fill)
+			.InitialSortMode(Column.GetInitialSortMode())
+			.FillWidth(Column.GetInitialWidth())
+			.HeaderContent()
+			[
+				SNew(SBox)
+				.HeightOverride(24.0f)
+				.Padding(FMargin(0.0f))
+				.VAlign(VAlign_Center)
+				[
+					SNew(STextBlock)
+					.Text(this, &SEntityEventsTableView::GetColumnHeaderText, Column.GetId())
+				]
+			];
+
+		int32 ColumnIndex = 0;
+		const int32 NewColumnPosition = Table->GetColumnPositionIndex(ColumnId);
+		const int32 NumColumns = HeaderRow->GetColumns().Num();
+		for (; ColumnIndex < NumColumns; ColumnIndex++)
+		{
+			const SHeaderRow::FColumn& CurrentColumn = HeaderRow->GetColumns()[ColumnIndex];
+			const int32 CurrentColumnPosition = Table->GetColumnPositionIndex(CurrentColumn.ColumnId);
+			if (NewColumnPosition < CurrentColumnPosition)
+			{
+				break;
+			}
+		}
+
+		HeaderRow->InsertColumn(ColumnArgs, ColumnIndex);
+	}
+
+	FText SEntityEventsTableView::GetColumnHeaderText(const FName ColumnId) const
+	{
+		const UE::Insights::FTableColumn& Column = *Table->FindColumnChecked(ColumnId);
+		return Column.GetShortName();
+	}
+
+	void SEntityEventsTableView::InsightsManager_OnSessionChanged()
+	{
+		IUnrealInsightsModule& UnrealInsightsModule = FModuleManager::LoadModuleChecked<IUnrealInsightsModule>("TraceInsights");
+
+		TSharedPtr<const TraceServices::IAnalysisSession> CurrentSession = UnrealInsightsModule.GetAnalysisSession();
+		
+		Session = CurrentSession;
+		Reset();
+	}
+
+	void SEntityEventsTableView::Reset()
+	{
+	}
+}
+
+#undef LOCTEXT_NAMESPACE
\ No newline at end of file
diff --git a/Engine/Plugins/MassInsights/Source/MassInsightsUI/Private/MassInsightsUI/Widgets/SEntityEventsTableView.h b/Engine/Plugins/MassInsights/Source/MassInsightsUI/Private/MassInsightsUI/Widgets/SEntityEventsTableView.h
new file mode 100644
--- /dev/null
+++ b/Engine/Plugins/MassInsights/Source/MassInsightsUI/Private/MassInsightsUI/Widgets/SEntityEventsTableView.h
@@ -0,0 +1,95 @@
+﻿// Copyright Epic Games, Inc. All Rights Reserved.
+
+#pragma once
+#include "Common.h"
+#include "SFragmentTableRow.h"
+#include "Framework/Commands/UICommandList.h"
+#include "InsightsCore/Table/ViewModels/BaseTreeNode.h"
+#include "Widgets/SCompoundWidget.h"
+#include "Widgets/Views/STreeView.h"
+#include "Widgets/Views/STableRow.h"
+#include "Widgets/Views/SListView.h"
+#include "InsightsCore/Table/ViewModels/Table.h"
+
+namespace MassInsightsUI
+{
+	struct FOnSelectedEntityEventParams
+	{
+		uint64 ProviderEventIndex;
+	};
+	DECLARE_DELEGATE_OneParam(FOnSelectedEntityEvent, const FOnSelectedEntityEventParams&);
+
+	/**
+	 * Displays a list of events for the given entities between a configured time period
+	 */
+	class SEntityEventsTableView : public SCompoundWidget
+	{
+	public:
+		SLATE_BEGIN_ARGS(SEntityEventsTableView) {}
+			SLATE_EVENT(FOnSelectedEntityEvent, OnEntityEventSelected)
+		SLATE_END_ARGS()
+
+		SEntityEventsTableView();
+		
+		void Construct(const FArguments& args);
+		void SetEntities(TConstArrayView<uint64> Entities);
+	private:
+	
+		class FEventNode : public UE::Insights::FBaseTreeNode
+		{
+			INSIGHTS_DECLARE_RTTI(FEventNode, UE::Insights::FBaseTreeNode)
+		public:
+			explicit FEventNode(const FName InName, bool bInIsGroup, uint64 InProviderEntityIndex);
+
+			uint64 GetProviderEventIndex() const;
+
+		private:
+			uint64 ProviderEntityIndex;
+			int32 Index;
+		};
+
+		friend class SEntityEventsTableRow;
+		
+		void InsightsManager_OnSessionChanged();
+
+		void TreeView_OnSelectionChanged(TSharedPtr<FEventNode> EventNode, ESelectInfo::Type Arg) const;
+		TSharedRef<ITableRow> TreeView_OnGenerateRow(TSharedPtr<FEventNode> NodePtr, const TSharedRef<STableViewBase>& OwnerTable);
+		void TreeView_OnGetChildren(TSharedPtr<FEventNode> InParent, TArray<TSharedPtr<FEventNode>>& OutChildren);
+
+		void IncrementalUpdate();
+		virtual void Tick(const FGeometry& AllottedGeometry, const double InCurrentTime, const float InDeltaTime) override;
+		void ShowColumn(const FName ColumnId);
+		FText GetColumnHeaderText(const FName ColumnId) const;
+
+		void Reset();
+
+		struct FEntityEventEntry
+		{
+			// Handle to link back to the data stored in the provider
+			uint64 ProviderEventIndex;
+			
+			double EventTime;
+			uint64 EntityId;
+
+			// To appease treeview
+			TSharedPtr<FEventNode> Handle;
+		};
+		
+		TSharedPtr<UE::Insights::FTable> Table;
+		TSharedPtr<const TraceServices::IAnalysisSession> Session;
+		TSharedPtr<STreeView<TSharedPtr<FEventNode>>> TreeView;
+		TSharedPtr<SHeaderRow> HeaderRow;
+		
+		TArray<TSharedPtr<FEventNode>> Events;
+		TArray<TSharedPtr<FEventNode>> FilteredSortedEvents;
+
+		bool bRebuildTable = true;
+		// Keeps track of the incremental progress of events to process
+		uint64 NextEventIndex = 0;
+		
+		TArray<uint64> Entities;
+		
+		MassInsightsUI::FOnSelectedArchetype OnArchetypeSelected;
+		FOnSelectedEntityEvent OnEntityEventSelected;
+	};
+}
diff --git a/Engine/Plugins/MassInsights/Source/MassInsightsUI/Private/MassInsightsUI/Widgets/SFragmentTableRow.cpp b/Engine/Plugins/MassInsights/Source/MassInsightsUI/Private/MassInsightsUI/Widgets/SFragmentTableRow.cpp
new file mode 100644
--- /dev/null
+++ b/Engine/Plugins/MassInsights/Source/MassInsightsUI/Private/MassInsightsUI/Widgets/SFragmentTableRow.cpp
@@ -0,0 +1,19 @@
+﻿// Copyright Epic Games, Inc. All Rights Reserved.
+
+#include "SFragmentTableRow.h"
+
+#include "Widgets/Text/STextBlock.h"
+
+void MassInsights::SFragmentTableRow::Construct(const FArguments& InArgs, const TSharedRef<STableViewBase>& InOwnerTableView)
+{
+	TablePtr = InArgs._TablePtr;
+	FragmentInfoPtr = InArgs._FragmentInfoPtr;
+	SetEnabled(true);
+	Super::Construct(Super::FArguments(), InOwnerTableView);
+}
+
+TSharedRef<SWidget> MassInsights::SFragmentTableRow::GenerateWidgetForColumn(const FName& InColumnName)
+{
+	return SNew( STextBlock )
+		.Text(FText::FromString( FragmentInfoPtr->Name ));
+}
diff --git a/Engine/Plugins/MassInsights/Source/MassInsightsUI/Private/MassInsightsUI/Widgets/SFragmentTableRow.h b/Engine/Plugins/MassInsights/Source/MassInsightsUI/Private/MassInsightsUI/Widgets/SFragmentTableRow.h
new file mode 100644
--- /dev/null
+++ b/Engine/Plugins/MassInsights/Source/MassInsightsUI/Private/MassInsightsUI/Widgets/SFragmentTableRow.h
@@ -0,0 +1,37 @@
+﻿// Copyright Epic Games, Inc. All Rights Reserved.
+
+#pragma once
+#include "InsightsCore/Table/ViewModels/Table.h"
+#include "MassInsightsAnalysis/Model/MassInsights.h"
+#include "Widgets/Views/STableRow.h"
+#include "Widgets/Views/SListView.h"
+
+namespace MassInsights
+{
+	struct FMassFragmentRowNodeData
+	{
+		uint64 Id;
+		FString Name;
+		uint32 Size;
+		MassInsightsAnalysis::EFragmentType Type;
+	};
+	using MassFragmentInfoPtr = TSharedPtr<FMassFragmentRowNodeData>;
+	class SFragmentTableRow : public SMultiColumnTableRow< MassFragmentInfoPtr >
+	{
+		using Super = SMultiColumnTableRow< MassFragmentInfoPtr >;
+	public:
+		SLATE_BEGIN_ARGS(SFragmentTableRow) {}
+			SLATE_ARGUMENT(TSharedPtr<UE::Insights::FTable>, TablePtr)
+			SLATE_ARGUMENT(MassFragmentInfoPtr, FragmentInfoPtr)
+		SLATE_END_ARGS()
+
+		void Construct(const FArguments& InArgs, const TSharedRef<STableViewBase>& InOwnerTableView);
+
+		virtual TSharedRef<SWidget> GenerateWidgetForColumn(const FName& InColumnName) override;
+
+	private:
+		TSharedPtr<UE::Insights::FTable> TablePtr;
+		MassFragmentInfoPtr FragmentInfoPtr;
+		
+	};
+}
diff --git a/Engine/Plugins/MassInsights/Source/MassInsightsUI/Private/MassInsightsUI/Widgets/SFragmentTableView.cpp b/Engine/Plugins/MassInsights/Source/MassInsightsUI/Private/MassInsightsUI/Widgets/SFragmentTableView.cpp
new file mode 100644
--- /dev/null
+++ b/Engine/Plugins/MassInsights/Source/MassInsightsUI/Private/MassInsightsUI/Widgets/SFragmentTableView.cpp
@@ -0,0 +1,222 @@
+﻿// Copyright Epic Games, Inc. All Rights Reserved.
+
+#include "SFragmentTableView.h"
+
+#include "Common/ProviderLock.h"
+#include "Insights/IInsightsManager.h"
+#include "InsightsCore/Table/ViewModels/TableCellValueFormatter.h"
+#include "InsightsCore/Table/ViewModels/TableCellValueGetter.h"
+#include "InsightsCore/Table/ViewModels/TableCellValueSorter.h"
+#include "InsightsCore/Table/ViewModels/TableColumn.h"
+#include "Modules/ModuleManager.h"
+
+#define LOCTEXT_NAMESPACE "MassInsights::SFragmentTableView"
+
+namespace MassInsights
+{
+	SFragmentTableView::SFragmentTableView()
+		: Table(MakeShared<UE::Insights::FTable>())
+	{
+	}
+
+	SFragmentTableView::~SFragmentTableView()
+	{
+	}
+
+	void SFragmentTableView::Construct(const FArguments& InArgs)
+	{
+		ChildSlot
+		[
+			SNew(SVerticalBox)
+
+			+SVerticalBox::Slot()
+			.VAlign(VAlign_Fill)
+			.Padding(2.0f, 2.0f, 2.0f, 2.0f)
+			[
+				SAssignNew(TreeView, STreeView<MassInsights::MassFragmentInfoPtr>)
+				.SelectionMode(ESelectionMode::Single)
+				.TreeItemsSource(&FragmentInfos)
+				.OnGetChildren(this, &SFragmentTableView::TreeView_OnGetChildren)
+				.OnGenerateRow(this, &SFragmentTableView::TreeView_OnGenerateRow)
+				.HeaderRow
+				(
+					SAssignNew(TreeViewHeaderRow, SHeaderRow)
+					.Visibility(EVisibility::Visible)
+				)
+			]
+		];
+
+		IUnrealInsightsModule& UnrealInsightsModule = FModuleManager::LoadModuleChecked<IUnrealInsightsModule>("TraceInsights");
+		TSharedPtr<::Insights::IInsightsManager> InsightsManager = UnrealInsightsModule.GetInsightsManager();
+		if (InsightsManager.IsValid())
+		{
+			InsightsManager->GetSessionChangedEvent().AddSP(this, &SFragmentTableView::InsightsManager_OnSessionChanged);
+			InsightsManager->GetSessionAnalysisCompletedEvent().AddSP(this, &SFragmentTableView::InsightsManager_OnSessionAnalysisComplete);
+		}
+		
+		TSharedPtr<const TraceServices::IAnalysisSession> CurrentSession = UnrealInsightsModule.GetAnalysisSession();
+		Session = CurrentSession;
+
+		TArray<TSharedRef<UE::Insights::FTableColumn>> Columns;
+
+		{
+			using namespace UE::Insights;
+			TSharedRef<FTableColumn> ColumnRef = MakeShared<FTableColumn>(FName(TEXT("NameColumn")));
+			FTableColumn& Column = *ColumnRef;
+			Column.SetShortName(LOCTEXT("XX", "Name"));
+			Column.SetTitleName(LOCTEXT("XX", "Name"));
+			Column.SetDescription(LOCTEXT("XX", "Name"));
+			Column.SetHorizontalAlignment(HAlign_Left);
+			Column.SetInitialWidth(206.0f);
+			Column.SetMinWidth(42.0f);
+			Column.SetFlags(ETableColumnFlags::ShouldBeVisible);
+			Column.SetDataType(ETableCellDataType::Text);
+			TSharedRef<ITableCellValueGetter> Getter = MakeShared<FDisplayNameValueGetter>();
+			Column.SetValueGetter(Getter);
+
+			TSharedRef<ITableCellValueFormatter> Formatter = MakeShared<FTextValueFormatter>();
+			Column.SetValueFormatter(Formatter);
+
+			TSharedRef<ITableCellValueSorter> Sorter = MakeShared<FSorterByName>(ColumnRef);
+			Column.SetValueSorter(Sorter);
+			
+			Column.Show();
+			
+			Columns.Add(ColumnRef);
+		}
+		
+		Table->SetColumns(Columns);
+
+		for (const TSharedRef<UE::Insights::FTableColumn>& ColumnRef : Table->GetColumns())
+		{
+			if (ColumnRef->ShouldBeVisible())
+			{
+				ShowColumn(ColumnRef->GetId());
+			}
+		}
+		
+	}
+
+	void SFragmentTableView::Tick(const FGeometry& AllottedGeometry, const double InCurrentTime, const float InDeltaTime)
+	{
+		SCompoundWidget::Tick(AllottedGeometry, InCurrentTime, InDeltaTime);
+
+		// Incremental update
+		RebuildTree(false);
+	}
+
+	void SFragmentTableView::InsightsManager_OnSessionChanged()
+	{
+		IUnrealInsightsModule& UnrealInsightsModule = FModuleManager::LoadModuleChecked<IUnrealInsightsModule>("TraceInsights");
+
+		TSharedPtr<const TraceServices::IAnalysisSession> CurrentSession = UnrealInsightsModule.GetAnalysisSession();
+		
+		Session = CurrentSession;
+		Reset();
+	}
+
+	void SFragmentTableView::InsightsManager_OnSessionAnalysisComplete()
+	{
+	}
+
+	TSharedRef<ITableRow> SFragmentTableView::TreeView_OnGenerateRow(
+		MassInsights::MassFragmentInfoPtr NodePtr,
+	    const TSharedRef<STableViewBase>& OwnerTable)
+	{
+		TSharedRef<ITableRow> TableRow =
+			SNew(MassInsights::SFragmentTableRow, OwnerTable)
+			.FragmentInfoPtr(NodePtr);
+	
+		return TableRow;
+	}
+
+	void SFragmentTableView::TreeView_OnGetChildren(
+		MassInsights::MassFragmentInfoPtr InParent,
+		TArray<MassInsights::MassFragmentInfoPtr>& OutChildren)
+	{
+		return;
+	}
+
+	void SFragmentTableView::Reset()
+	{
+		RebuildTree(true);
+	}
+
+	void SFragmentTableView::RebuildTree(bool bResync)
+	{
+		if (bResync)
+		{
+			FragmentInfos.Reset();
+		}
+		const uint32 PreviousFragmentCount = FragmentInfos.Num();
+		
+		if (Session.IsValid())
+		{
+			const MassInsightsAnalysis::IMassInsightsProvider& Provider = MassInsightsAnalysis::ReadMassInsightsProvider(*Session);
+
+			TraceServices::FProviderReadScopeLock ProviderReadScopeLock(Provider);
+			const uint32 FragmentCount = static_cast<uint32>(Provider.GetFragmentCount());
+
+			if (FragmentCount != PreviousFragmentCount)
+			{
+				FragmentInfos.SetNum(FragmentCount);
+				Provider.EnumerateFragments([this](const MassInsightsAnalysis::FMassFragmentInfo& FragmentInfo, int32 Index)
+				{
+					MassFragmentInfoPtr InfoPtr = MakeShared<MassFragmentInfoPtr::ElementType>();
+					InfoPtr->Id = FragmentInfo.Id;
+					InfoPtr->Name = FragmentInfo.Name;
+					InfoPtr->Size = FragmentInfo.Size;
+					InfoPtr->Type = FragmentInfo.Type;
+					FragmentInfos[Index] = InfoPtr;
+				}, PreviousFragmentCount);
+				
+				TreeView->RebuildList();
+			}
+		}
+	}
+
+	void SFragmentTableView::ShowColumn(const FName ColumnId)
+	{
+		UE::Insights::FTableColumn& Column = *Table->FindColumnChecked(ColumnId);
+		Column.Show();
+
+		SHeaderRow::FColumn::FArguments ColumnArgs;
+		ColumnArgs
+			.ColumnId(Column.GetId())
+			.DefaultLabel(Column.GetShortName())
+			.HAlignHeader(Column.GetHorizontalAlignment())
+			.VAlignHeader(VAlign_Center)
+			.HAlignCell(HAlign_Fill)
+			.VAlignCell(VAlign_Fill)
+			.InitialSortMode(Column.GetInitialSortMode())
+			.FillWidth(Column.GetInitialWidth())
+			.HeaderContent()
+			[
+				SNew(SBox)
+				.HeightOverride(24.0f)
+				.Padding(FMargin(0.0f))
+				.VAlign(VAlign_Center)
+				[
+					SNew(STextBlock)
+					.Text(FText::FromString(TEXT("Header Text")))
+				]
+			];
+
+		int32 ColumnIndex = 0;
+		const int32 NewColumnPosition = Table->GetColumnPositionIndex(ColumnId);
+		const int32 NumColumns = TreeViewHeaderRow->GetColumns().Num();
+		for (; ColumnIndex < NumColumns; ColumnIndex++)
+		{
+			const SHeaderRow::FColumn& CurrentColumn = TreeViewHeaderRow->GetColumns()[ColumnIndex];
+			const int32 CurrentColumnPosition = Table->GetColumnPositionIndex(CurrentColumn.ColumnId);
+			if (NewColumnPosition < CurrentColumnPosition)
+			{
+				break;
+			}
+		}
+
+		TreeViewHeaderRow->InsertColumn(ColumnArgs, ColumnIndex);
+	}
+}
+
+#undef LOCTEXT_NAMESPACE
\ No newline at end of file
diff --git a/Engine/Plugins/MassInsights/Source/MassInsightsUI/Private/MassInsightsUI/Widgets/SFragmentTableView.h b/Engine/Plugins/MassInsights/Source/MassInsightsUI/Private/MassInsightsUI/Widgets/SFragmentTableView.h
new file mode 100644
--- /dev/null
+++ b/Engine/Plugins/MassInsights/Source/MassInsightsUI/Private/MassInsightsUI/Widgets/SFragmentTableView.h
@@ -0,0 +1,42 @@
+﻿#pragma once
+#include "SFragmentTableRow.h"
+#include "Framework/Commands/UICommandList.h"
+#include "Widgets/SCompoundWidget.h"
+#include "Widgets/Views/STreeView.h"
+
+namespace MassInsights
+{
+
+class SFragmentTableView : public SCompoundWidget
+{
+public:
+	SFragmentTableView();
+	virtual ~SFragmentTableView();
+	
+	SLATE_BEGIN_ARGS(SFragmentTableView) {}
+	SLATE_END_ARGS()
+
+	void Construct(const FArguments& InArgs);
+private:
+	virtual void Tick(const FGeometry& AllottedGeometry, const double InCurrentTime, const float InDeltaTime) override;
+	void InsightsManager_OnSessionChanged();
+	void InsightsManager_OnSessionAnalysisComplete();
+	TSharedRef<ITableRow> TreeView_OnGenerateRow(MassInsights::MassFragmentInfoPtr NodePtr, const TSharedRef<STableViewBase>& OwnerTable);
+	void TreeView_OnGetChildren(MassInsights::MassFragmentInfoPtr InParent, TArray<MassInsights::MassFragmentInfoPtr>& OutChildren);
+	
+	void Reset();
+	void RebuildTree(bool bResync);
+	void ShowColumn(const FName ColumnId);
+private:
+	TSharedPtr<UE::Insights::FTable> Table;
+	TSharedPtr<const TraceServices::IAnalysisSession> Session;
+	TSharedPtr<FUICommandList> CommandList;
+	
+	TSharedPtr<STreeView<MassInsights::MassFragmentInfoPtr>> TreeView;
+	TSharedPtr<SHeaderRow> TreeViewHeaderRow;
+
+	TArray<MassInsights::MassFragmentInfoPtr> FragmentInfos;
+	TArray<MassInsights::MassFragmentInfoPtr> FilteredTreeItems;
+};
+
+}
diff --git a/Engine/Plugins/MassInsights/Source/MassInsightsUI/Private/MassInsightsUI/Widgets/SMassInsightsAnalysisTab.cpp b/Engine/Plugins/MassInsights/Source/MassInsightsUI/Private/MassInsightsUI/Widgets/SMassInsightsAnalysisTab.cpp
new file mode 100644
--- /dev/null
+++ b/Engine/Plugins/MassInsights/Source/MassInsightsUI/Private/MassInsightsUI/Widgets/SMassInsightsAnalysisTab.cpp
@@ -0,0 +1,227 @@
+﻿// Copyright Epic Games, Inc. All Rights Reserved.
+
+#include "SMassInsightsAnalysisTab.h"
+
+#include "SArchetypeDetails.h"
+#include "SEntityEventsTableView.h"
+#include "SFragmentTableView.h"
+#include "Common/ProviderLock.h"
+#include "Insights/ITimingViewSession.h"
+#include "Insights/ViewModels/ITimingEvent.h"
+#include "Styling/AppStyle.h"
+#include "Widgets/SBoxPanel.h"
+#include "Widgets/Input/SCheckBox.h"
+#include "Widgets/Layout/SWidgetSwitcher.h"
+#include "Widgets/Text/STextBlock.h"
+#include "Widgets/SNullWidget.h"
+#include "MassInsightsUI/Widgets/SEntityEventsTableView.h"
+
+#define LOCTEXT_NAMESPACE "SMassInsightsAnalysisTab"
+
+namespace MassInsights
+{
+	void SMassInsightsAnalysisTab::Construct(const FArguments& InArgs)
+	{
+		TSharedPtr<SHorizontalBox> ViewModeSelectorsBox;
+		TSharedPtr<SWidgetSwitcher> ViewModeSwitcher;
+		
+
+		TSharedRef<MassInsightsUI::SArchetypeDetails> ArchetypeDetailsRef =
+						SNew(MassInsightsUI::SArchetypeDetails);
+		ArchetypesDetails = ArchetypeDetailsRef;
+
+		TSharedRef<SEntityEventAggregationTableView> EntityTimelineTableViewRef =
+			SNew(SEntityEventAggregationTableView)
+			.OnArchetypeSelected_Lambda([this](uint64 ArchetypeID)
+			{
+				TSharedPtr<MassInsightsUI::SArchetypeDetails> Pinned = ArchetypesDetails.Pin();
+				if (Pinned) { Pinned->SetArchetype(ArchetypeID); }
+			})
+			.OnRowSelected_Lambda([this](const FEntityEventSummaryRowSelectedParams& RowData)
+			{
+				TSharedPtr<MassInsightsUI::SEntityEventsTableView> Pinned = TableView.Pin();
+				if (Pinned)
+				{
+					Pinned->SetEntities(MakeConstArrayView(&RowData.EntityID, 1));
+				}
+
+				// Add a time marker when the row is selected - will highlight the first event found
+				// for that entity
+				if (TimingViewSession)
+				{
+					if (RowData.IsSelected)
+					{
+						TimingViewSession->SetTimeMarker(RowData.FirstEventTime);
+					}
+				}
+			});
+		EntityTimelineTableView = EntityTimelineTableViewRef;
+		
+		TSharedRef<MassInsightsUI::SEntityEventsTableView> JourneyTableViewRef =
+			SNew(MassInsightsUI::SEntityEventsTableView)
+			.OnEntityEventSelected_Lambda([this](const MassInsightsUI::FOnSelectedEntityEventParams& Params)
+			{
+				TSharedPtr<MassInsightsUI::SArchetypeDetails> Pinned = ArchetypesDetails.Pin();
+				if (Pinned)
+				{
+					TValueOrError<MassInsightsAnalysis::FMassEntityEventRecord,void> Event = MakeError();
+					const MassInsightsAnalysis::IMassInsightsProvider& Provider = MassInsightsAnalysis::ReadMassInsightsProvider(*AnalysisSession);
+					{
+						TraceServices::FProviderReadScopeLock ProviderReadScopeLock(Provider);
+
+						Event = Provider.GetEntityEvent(Params.ProviderEventIndex);
+					}
+
+					if (Event.HasValue())
+					{
+						Pinned->SetArchetype(Event.GetValue().ArchetypeID);
+						
+						if (TimingViewSession)
+						{
+							TimingViewSession->SetTimeMarker(Event.GetValue().Time);
+						}
+					}
+					
+					
+				}
+			});
+		TableView = JourneyTableViewRef;
+		
+		ChildSlot
+		[
+			SNew(SVerticalBox)
+
+			// view mode selection
+			+SVerticalBox::Slot()
+			.AutoHeight()
+			.Padding(8,2)
+			[
+				SAssignNew(ViewModeSelectorsBox, SHorizontalBox)
+			]
+
+			// active view mode panel
+			+SVerticalBox::Slot()
+			.VAlign(VAlign_Fill)
+			.FillHeight(1.0f)
+			.Padding(3)
+			[
+				SAssignNew(ViewModeSwitcher, SWidgetSwitcher)
+				.WidgetIndex(this, &SMassInsightsAnalysisTab::GetSelectedViewModeIndex)
+			]
+		];
+
+		
+		// helper to create widgets for a view mode
+		auto AddViewMode = [this, &ViewModeSelectorsBox, &ViewModeSwitcher]( EViewMode InViewMode, const FText& InText)
+		{
+			// add mode selection button
+			ViewModeSelectorsBox->AddSlot()
+			.AutoWidth()
+			.Padding(4)
+			[
+				SNew(SCheckBox)
+				.Style(&FAppStyle::GetWidgetStyle<FCheckBoxStyle>("ToggleButtonCheckBox"))
+				.IsChecked( this, &SMassInsightsAnalysisTab::IsViewModeSelected, InViewMode )
+				.OnCheckStateChanged( this, &SMassInsightsAnalysisTab::OnViewModeCheckStateChange, InViewMode )
+				[
+					SNew(STextBlock)
+						.Text(InText)
+				]
+			];
+
+			// return a new slot for the mode's panel
+			return ViewModeSwitcher->AddSlot();
+		};
+
+		
+		
+		AddViewMode(EViewMode::EntityEvents, LOCTEXT("ViewModeTab_Entities", "Entity Events"))
+		[
+			SNew(SSplitter)
+			.Orientation(EOrientation::Orient_Horizontal)
+
+			+SSplitter::Slot()
+			.Value(0.6)
+			[
+				SNew(SBorder)
+				.Padding(4)
+				[
+					EntityTimelineTableViewRef
+				]
+			]
+
+			+SSplitter::Slot()
+			.Value(0.4)
+			[
+				SNew(SSplitter)
+				.Orientation(EOrientation::Orient_Vertical)
+				+SSplitter::Slot()
+				.Value(0.5)
+				[
+					SNew(SBorder)
+					.Padding(4)
+					[
+						JourneyTableViewRef
+					]
+				]
+				+SSplitter::Slot()
+				.Value(0.5)
+				[
+					SNew(SBorder)
+					.Padding(4)
+					[
+						ArchetypeDetailsRef
+					]
+				]
+			]
+		];
+		
+		AddViewMode(EViewMode::Fragments, LOCTEXT("ViewMode_Fragments", "Fragments"))
+		[
+			SAssignNew(FragmentTableView, SFragmentTableView)
+		];
+	}
+
+	void SMassInsightsAnalysisTab::SetSession(UE::Insights::Timing::ITimingViewSession* InTimingViewSession,
+		const TraceServices::IAnalysisSession* InAnalysisSession)
+	{
+		if (TimingViewSession == InTimingViewSession)
+		{
+			return;
+		}
+
+		if (TimingViewSession)
+		{
+			TimingViewSession->OnSelectionChanged().RemoveAll(this);
+			TimingViewSession->OnSelectedEventChanged().RemoveAll(this);
+		}
+		
+		TimingViewSession = InTimingViewSession;
+		AnalysisSession = InAnalysisSession;
+	}
+
+	int32 SMassInsightsAnalysisTab::GetSelectedViewModeIndex() const
+	{
+		return (int32)ViewMode;
+	}
+
+	ECheckBoxState SMassInsightsAnalysisTab::IsViewModeSelected(EViewMode InMode) const
+	{
+		return (InMode == ViewMode) ? ECheckBoxState::Checked : ECheckBoxState::Unchecked;
+	}
+
+	void SMassInsightsAnalysisTab::OnViewModeCheckStateChange(ECheckBoxState State, EViewMode Mode)
+	{
+		if (State == ECheckBoxState::Checked)
+		{
+			ViewMode = Mode;
+		}
+	}
+
+	bool SMassInsightsAnalysisTab::IsSessionSet() const
+	{
+		return TimingViewSession && AnalysisSession;
+	}
+}
+
+#undef LOCTEXT_NAMESPACE
\ No newline at end of file
diff --git a/Engine/Plugins/MassInsights/Source/MassInsightsUI/Private/MassInsightsUI/Widgets/SMassInsightsAnalysisTab.h b/Engine/Plugins/MassInsights/Source/MassInsightsUI/Private/MassInsightsUI/Widgets/SMassInsightsAnalysisTab.h
new file mode 100644
--- /dev/null
+++ b/Engine/Plugins/MassInsights/Source/MassInsightsUI/Private/MassInsightsUI/Widgets/SMassInsightsAnalysisTab.h
@@ -0,0 +1,64 @@
+﻿// Copyright Epic Games, Inc. All Rights Reserved.
+
+#pragma once
+
+#include "SArchetypeDetails.h"
+#include "SEntityEventAggregationTableView.h"
+#include "SFragmentTableView.h"
+#include "Widgets/SCompoundWidget.h"
+
+namespace MassInsightsUI
+{
+	class SEntityEventsTableView;
+}
+
+class ITimingEvent;
+
+namespace UE::Insights::Timing
+{
+	class ITimingViewSession;
+	enum class ETimeChangedFlags;
+}
+
+namespace TraceServices
+{
+	class IAnalysisSession;
+}
+
+namespace MassInsights
+{
+	class SMassInsightsAnalysisTab : public SCompoundWidget
+	{
+	public:
+		SLATE_BEGIN_ARGS(SMassInsightsAnalysisTab) {}
+		SLATE_END_ARGS()
+		
+		void Construct(const FArguments& InArgs);
+
+		void SetSession(
+			UE::Insights::Timing::ITimingViewSession* InTimingViewSession,
+			const TraceServices::IAnalysisSession* InAnalysisSession);
+		bool IsSessionSet() const;
+
+	private:
+		enum class EViewMode : uint8
+		{
+			EntityEvents,
+			Fragments,
+		};
+		
+		int32 GetSelectedViewModeIndex() const;
+		ECheckBoxState IsViewModeSelected(EViewMode ViewMode) const;
+		void OnViewModeCheckStateChange(ECheckBoxState State, EViewMode Mode);
+		
+		UE::Insights::Timing::ITimingViewSession* TimingViewSession = nullptr;
+		const TraceServices::IAnalysisSession* AnalysisSession = nullptr;
+
+		TSharedPtr<SFragmentTableView> FragmentTableView;
+		TSharedPtr<SEntityEventAggregationTableView> EntityTimelineTableView;
+		TWeakPtr<MassInsightsUI::SArchetypeDetails> ArchetypesDetails;
+		TWeakPtr<MassInsightsUI::SEntityEventsTableView> TableView;
+		
+		EViewMode ViewMode = EViewMode::EntityEvents;
+	};
+}
\ No newline at end of file
diff --git a/Engine/Plugins/MassInsights/Source/MassInsightsUI/Private/MassInsightsUIModule.cpp b/Engine/Plugins/MassInsights/Source/MassInsightsUI/Private/MassInsightsUIModule.cpp
new file mode 100644
--- /dev/null
+++ b/Engine/Plugins/MassInsights/Source/MassInsightsUI/Private/MassInsightsUIModule.cpp
@@ -0,0 +1,85 @@
+// Copyright Epic Games, Inc. All Rights Reserved.
+#include "MassInsightsUIModule.h"
+
+#include "Modules/ModuleManager.h"
+#include "Features/IModularFeatures.h"
+#include "Insights/ITimingViewExtender.h"
+#include "Insights/IUnrealInsightsModule.h"
+#include "MassInsightsUI/MassInsightsStyle.h"
+#include "MassInsightsUI/Widgets/SFragmentTableView.h"
+#include "MassInsightsUI/Widgets/SMassInsightsAnalysisTab.h"
+#include "Widgets/Docking/SDockTab.h"
+#include "Widgets/Text/STextBlock.h"
+
+#define LOCTEXT_NAMESPACE "MassInsightsModule"
+
+namespace MassInsightsUI
+{
+	FMassInsightsUIModule& FMassInsightsUIModule::Get()
+	{
+		return FModuleManager::LoadModuleChecked<FMassInsightsUIModule>("MassInsightsUI");
+	}
+	
+	void FMassInsightsUIModule::StartupModule()
+	{
+		IModularFeatures::Get().RegisterModularFeature(UE::Insights::Timing::TimingViewExtenderFeatureName, &TimingViewExtender);
+
+		IUnrealInsightsModule& UnrealInsightsModule = FModuleManager::LoadModuleChecked<IUnrealInsightsModule>("TraceInsights");
+		FOnRegisterMajorTabExtensions& TimingProfilerLayoutExtension = UnrealInsightsModule.OnRegisterMajorTabExtension(FInsightsManagerTabs::TimingProfilerTabId);
+		TimingProfilerLayoutExtension.AddRaw(this, &FMassInsightsUIModule::RegisterLayoutExtension);
+	}
+
+	void FMassInsightsUIModule::ShutdownModule()
+	{
+		IModularFeatures::Get().UnregisterModularFeature(UE::Insights::Timing::TimingViewExtenderFeatureName, &TimingViewExtender);
+	}
+
+	TSharedPtr<MassInsights::SMassInsightsAnalysisTab> FMassInsightsUIModule::GetAnalysisTab()
+	{
+		return AnalysisTab.Pin();
+	}
+
+	namespace Private
+	{
+		const FName FrameViewTab("SlateFrame2ViewTab");
+	}
+
+	void FMassInsightsUIModule::RegisterLayoutExtension(FInsightsMajorTabExtender& Extender)
+	{
+		InsightsTabManager = Extender.GetTabManager();
+
+		FMinorTabConfig& MinorTabConfig = Extender.AddMinorTabConfig();
+		MinorTabConfig.TabId = Private::FrameViewTab;
+		MinorTabConfig.TabLabel = LOCTEXT("MassInsightsTabTitle", "Mass Insights");
+		MinorTabConfig.TabTooltip = LOCTEXT("MassInsightsTabTitleTooltip", "Tooltip here");
+		MinorTabConfig.TabIcon = FSlateIcon(MassInsights::FMassInsightsStyle::Get().GetStyleSetName(), "MassInsights.Icon.Small");
+		MinorTabConfig.WorkspaceGroup = Extender.GetWorkspaceGroup();
+		MinorTabConfig.OnSpawnTab = FOnSpawnTab::CreateLambda([this](const FSpawnTabArgs& Args)
+			{
+				TSharedRef<MassInsights::SMassInsightsAnalysisTab> Content = SNew(MassInsights::SMassInsightsAnalysisTab);
+				AnalysisTab = Content;
+			
+				const TSharedRef<SDockTab> DockTab = SNew(SDockTab)
+					.TabRole(ETabRole::PanelTab)
+					[
+						SNew(SVerticalBox)
+						+SVerticalBox::Slot()
+						.Padding(FMargin(2.0))
+						[
+							Content
+						]					
+					];
+
+				return DockTab;
+			});
+
+		Extender.GetLayoutExtender().ExtendLayout(FTimingProfilerTabs::TimersID
+			, ELayoutExtensionPosition::Before
+			, FTabManager::FTab(Private::FrameViewTab, ETabState::ClosedTab));
+	}
+}
+
+
+IMPLEMENT_MODULE(MassInsightsUI::FMassInsightsUIModule, MassInsightsUI);
+
+#undef LOCTEXT_NAMESPACE
diff --git a/Engine/Plugins/MassInsights/Source/MassInsightsUI/Private/MassInsightsUIModule.h b/Engine/Plugins/MassInsights/Source/MassInsightsUI/Private/MassInsightsUIModule.h
new file mode 100644
--- /dev/null
+++ b/Engine/Plugins/MassInsights/Source/MassInsightsUI/Private/MassInsightsUIModule.h
@@ -0,0 +1,39 @@
+// Copyright Epic Games, Inc. All Rights Reserved.
+#pragma once
+
+#include "Modules/ModuleInterface.h"
+#include "MassInsightsTimingTrack.h"
+
+namespace MassInsights
+{
+	class SMassInsightsAnalysisTab;
+}
+
+struct FInsightsMajorTabExtender;
+class FTabManager;
+
+namespace MassInsightsUI
+{
+	
+
+class FMassInsightsUIModule : public IModuleInterface
+{
+public:
+	//~ Begin IModuleInterface interface
+	virtual void StartupModule() override;
+	virtual void ShutdownModule() override;
+	//~ End IModuleInterface interface
+	
+	static FMassInsightsUIModule& Get();
+
+	TSharedPtr<MassInsights::SMassInsightsAnalysisTab> GetAnalysisTab();
+
+private:
+	void RegisterLayoutExtension(FInsightsMajorTabExtender& Extender);
+	FMassInsightsSharedState TimingViewExtender;
+	TWeakPtr<FTabManager> InsightsTabManager;
+
+	TWeakPtr<MassInsights::SMassInsightsAnalysisTab> AnalysisTab;
+};
+
+} //namespace MassInsightsUI
\ No newline at end of file
diff --git a/Engine/Programs/UnrealInsights/Config/DefaultEngine.ini b/Engine/Programs/UnrealInsights/Config/DefaultEngine.ini
--- a/Engine/Programs/UnrealInsights/Config/DefaultEngine.ini
+++ b/Engine/Programs/UnrealInsights/Config/DefaultEngine.ini
@@ -9,6 +9,7 @@
 +ProgramEnabledPlugins=NetworkPredictionInsights
 +ProgramEnabledPlugins=VisualStudioSourceCodeAccess
 +ProgramEnabledPlugins=XcodeSourceCodeAccess
++ProgramEnabledPlugins=MassInsights
 
 [/Script/MacTargetPlatform.XcodeProjectSettings]
 TemplateMacPlist=(FilePath="/Engine/Source/Programs/UnrealInsights/Resources/Mac/Info.plist")
