Fix copyright notices

diff --git a/Engine/Plugins/MassInsights/Source/MassInsightsAnalysis/MassInsightsAnalysis.Build.cs b/Engine/Plugins/MassInsights/Source/MassInsightsAnalysis/MassInsightsAnalysis.Build.cs
--- a/Engine/Plugins/MassInsights/Source/MassInsightsAnalysis/MassInsightsAnalysis.Build.cs
+++ b/Engine/Plugins/MassInsights/Source/MassInsightsAnalysis/MassInsightsAnalysis.Build.cs
@@ -1,3 +1,5 @@
+// Copyright Epic Games, Inc. All Rights Reserved.
+
 namespace UnrealBuildTool.Rules
 {
 	public class MassInsightsAnalysis : ModuleRules
diff --git a/Engine/Plugins/MassInsights/Source/MassInsightsAnalysis/Private/Model/MassInsights.cpp b/Engine/Plugins/MassInsights/Source/MassInsightsAnalysis/Private/Model/MassInsights.cpp
--- a/Engine/Plugins/MassInsights/Source/MassInsightsAnalysis/Private/Model/MassInsights.cpp
+++ b/Engine/Plugins/MassInsights/Source/MassInsightsAnalysis/Private/Model/MassInsights.cpp
@@ -1,3 +1,5 @@
+// Copyright Epic Games, Inc. All Rights Reserved.
+
 #include "MassInsightsAnalysis/Model/MassInsights.h"
 #include "MassInsightsPrivate.h"
 
diff --git a/Engine/Plugins/MassInsights/Source/MassInsightsAnalysis/Private/Model/MassInsightsPrivate.h b/Engine/Plugins/MassInsights/Source/MassInsightsAnalysis/Private/Model/MassInsightsPrivate.h
--- a/Engine/Plugins/MassInsights/Source/MassInsightsAnalysis/Private/Model/MassInsightsPrivate.h
+++ b/Engine/Plugins/MassInsights/Source/MassInsightsAnalysis/Private/Model/MassInsightsPrivate.h
@@ -1,3 +1,5 @@
+// Copyright Epic Games, Inc. All Rights Reserved.
+
 #pragma once
 
 #include "Common/ProviderLock.h"
diff --git a/Engine/Plugins/MassInsights/Source/MassInsightsAnalysis/Public/MassInsightsAnalysis/Model/MassInsights.h b/Engine/Plugins/MassInsights/Source/MassInsightsAnalysis/Public/MassInsightsAnalysis/Model/MassInsights.h
--- a/Engine/Plugins/MassInsights/Source/MassInsightsAnalysis/Public/MassInsightsAnalysis/Model/MassInsights.h
+++ b/Engine/Plugins/MassInsights/Source/MassInsightsAnalysis/Public/MassInsightsAnalysis/Model/MassInsights.h
@@ -1,3 +1,5 @@
+// Copyright Epic Games, Inc. All Rights Reserved.
+
 #pragma once
 
 #include "Common/PagedArray.h" // TraceServices
diff --git a/Engine/Plugins/MassInsights/Source/MassInsightsUI/MassInsightsUI.Build.cs b/Engine/Plugins/MassInsights/Source/MassInsightsUI/MassInsightsUI.Build.cs
--- a/Engine/Plugins/MassInsights/Source/MassInsightsUI/MassInsightsUI.Build.cs
+++ b/Engine/Plugins/MassInsights/Source/MassInsightsUI/MassInsightsUI.Build.cs
@@ -1,3 +1,5 @@
+// Copyright Epic Games, Inc. All Rights Reserved.
+
 namespace UnrealBuildTool.Rules
 {
 	public class MassInsightsUI : ModuleRules
diff --git a/Engine/Plugins/MassInsights/Source/MassInsightsUI/Private/MassInsightsUI/Widgets/SFragmentTableView.h b/Engine/Plugins/MassInsights/Source/MassInsightsUI/Private/MassInsightsUI/Widgets/SFragmentTableView.h
--- a/Engine/Plugins/MassInsights/Source/MassInsightsUI/Private/MassInsightsUI/Widgets/SFragmentTableView.h
+++ b/Engine/Plugins/MassInsights/Source/MassInsightsUI/Private/MassInsightsUI/Widgets/SFragmentTableView.h
@@ -1,4 +1,6 @@
-﻿#pragma once
+﻿// Copyright Epic Games, Inc. All Rights Reserved.
+
+#pragma once
 #include "SFragmentTableRow.h"
 #include "Framework/Commands/UICommandList.h"
 #include "Widgets/SCompoundWidget.h"
