Fix missing LOCTEXT key

diff --git a/Engine/Plugins/MassInsights/Source/MassInsightsUI/Private/MassInsightsUI/Widgets/SEntityEventAggregationTableView.cpp b/Engine/Plugins/MassInsights/Source/MassInsightsUI/Private/MassInsightsUI/Widgets/SEntityEventAggregationTableView.cpp
--- a/Engine/Plugins/MassInsights/Source/MassInsightsUI/Private/MassInsightsUI/Widgets/SEntityEventAggregationTableView.cpp
+++ b/Engine/Plugins/MassInsights/Source/MassInsightsUI/Private/MassInsightsUI/Widgets/SEntityEventAggregationTableView.cpp
@@ -589,7 +589,7 @@ namespace MassInsights
 
 	FText SEntityEventAggregationTableView::GetStatusBarText() const
 	{
-		return FText::Format(LOCTEXT("", "-- Entities:{0} Events:{1}/{2}--"), FilteredTreeViewRows.Num(), EventsProcessed, EstimatedEventCount);
+		return FText::Format(LOCTEXT("EntityStatusBarText", "-- Entities:{0} Events:{1}/{2}--"), FilteredTreeViewRows.Num(), EventsProcessed, EstimatedEventCount);
 	}
 
 }
