LOCTEXT fixes for MassInsights

diff --git a/Engine/Plugins/MassInsights/Source/MassInsightsUI/Private/MassInsightsUI/Widgets/SArchetypeDetails.cpp b/Engine/Plugins/MassInsights/Source/MassInsightsUI/Private/MassInsightsUI/Widgets/SArchetypeDetails.cpp
--- a/Engine/Plugins/MassInsights/Source/MassInsightsUI/Private/MassInsightsUI/Widgets/SArchetypeDetails.cpp
+++ b/Engine/Plugins/MassInsights/Source/MassInsightsUI/Private/MassInsightsUI/Widgets/SArchetypeDetails.cpp
@@ -129,9 +129,9 @@ namespace MassInsightsUI
 			using namespace UE::Insights;
 			TSharedRef<UE::Insights::FTableColumn> ColumnRef = MakeShared<FTableColumn>(NAME_COLUMN_FragmentType);
 			FTableColumn& Column = *ColumnRef;
-			Column.SetShortName(LOCTEXT("FragmentName", "Type"));
-			Column.SetTitleName(LOCTEXT("FragmentName", "Type"));
-			Column.SetDescription(LOCTEXT("FragmentName", "Type"));
+			Column.SetShortName(LOCTEXT("FragmentType", "Type"));
+			Column.SetTitleName(LOCTEXT("FragmentType", "Type"));
+			Column.SetDescription(LOCTEXT("FragmentType", "Type"));
 			Column.SetHorizontalAlignment(HAlign_Left);
 			Column.SetInitialWidth(206.0f);
 			Column.SetMinWidth(42.0f);
diff --git a/Engine/Plugins/MassInsights/Source/MassInsightsUI/Private/MassInsightsUI/Widgets/SEntityEventAggregationTableView.cpp b/Engine/Plugins/MassInsights/Source/MassInsightsUI/Private/MassInsightsUI/Widgets/SEntityEventAggregationTableView.cpp
--- a/Engine/Plugins/MassInsights/Source/MassInsightsUI/Private/MassInsightsUI/Widgets/SEntityEventAggregationTableView.cpp
+++ b/Engine/Plugins/MassInsights/Source/MassInsightsUI/Private/MassInsightsUI/Widgets/SEntityEventAggregationTableView.cpp
@@ -94,7 +94,7 @@ namespace MassInsights
 		if (InColumnName == COLUMN_EntityID)
 		{
 			return SNew( STextBlock )
-				.Text(FText::Format(LOCTEXT("EntityID", "{0}"), AggregationRecord.EntityID ));
+				.Text(FText::Format(LOCTEXT("EntityID_Cell", "{0}"), AggregationRecord.EntityID ));
 		}
 		else if (InColumnName == COLUMN_Archetype)
 		{
