[MassInsights] Fixed deprecated FName for FBaseTreeNode constructor.

#rnx
#Synced-CL 52443010

diff --git a/Engine/Plugins/MassInsights/Source/MassInsightsUI/Private/MassInsightsUI/Widgets/SEntityEventAggregationTableView.cpp b/Engine/Plugins/MassInsights/Source/MassInsightsUI/Private/MassInsightsUI/Widgets/SEntityEventAggregationTableView.cpp
--- a/Engine/Plugins/MassInsights/Source/MassInsightsUI/Private/MassInsightsUI/Widgets/SEntityEventAggregationTableView.cpp
+++ b/Engine/Plugins/MassInsights/Source/MassInsightsUI/Private/MassInsightsUI/Widgets/SEntityEventAggregationTableView.cpp
@@ -28,11 +28,16 @@ namespace MassInsights
 	INSIGHTS_IMPLEMENT_RTTI(FEntityEventAggregateRecordHandle)
 	
 	FEntityEventAggregateRecordHandle::FEntityEventAggregateRecordHandle(int32 InRecordIndex)
-		: FBaseTreeNode(NAME_None, false /* bInIsGroup */)
+		: FBaseTreeNode(false /* bInIsGroup */)
 		, RecordIndex(InRecordIndex)
 	{
 	}
 
+	const FText FEntityEventAggregateRecordHandle::GetDisplayName() const
+	{
+		return FText::FromString(FString::Printf(TEXT("record_%d"), RecordIndex));
+	}
+
 	int32 FEntityEventAggregateRecordHandle::GetRecordIndex() const
 	{
 		return RecordIndex;
diff --git a/Engine/Plugins/MassInsights/Source/MassInsightsUI/Private/MassInsightsUI/Widgets/SEntityEventAggregationTableView.h b/Engine/Plugins/MassInsights/Source/MassInsightsUI/Private/MassInsightsUI/Widgets/SEntityEventAggregationTableView.h
--- a/Engine/Plugins/MassInsights/Source/MassInsightsUI/Private/MassInsightsUI/Widgets/SEntityEventAggregationTableView.h
+++ b/Engine/Plugins/MassInsights/Source/MassInsightsUI/Private/MassInsightsUI/Widgets/SEntityEventAggregationTableView.h
@@ -25,9 +25,10 @@ namespace MassInsights
 	public:
 		explicit FEntityEventAggregateRecordHandle(int32 InRecordIndex);
 
+		virtual const FText GetDisplayName() const override;
 		int32 GetRecordIndex() const;
+
 	private:
-		
 		// Index used to reference event aggregation data
 		int32 RecordIndex;
 	};
diff --git a/Engine/Plugins/MassInsights/Source/MassInsightsUI/Private/MassInsightsUI/Widgets/SEntityEventsTableView.cpp b/Engine/Plugins/MassInsights/Source/MassInsightsUI/Private/MassInsightsUI/Widgets/SEntityEventsTableView.cpp
--- a/Engine/Plugins/MassInsights/Source/MassInsightsUI/Private/MassInsightsUI/Widgets/SEntityEventsTableView.cpp
+++ b/Engine/Plugins/MassInsights/Source/MassInsightsUI/Private/MassInsightsUI/Widgets/SEntityEventsTableView.cpp
@@ -142,12 +142,17 @@ namespace MassInsightsUI
 		OnArchetypeSelected.ExecuteIfBound(ArchetypeID);
 	}
 	
-	SEntityEventsTableView::FEventNode::FEventNode(const FName InName, bool bInIsGroup, uint64 InProviderEntityIndex)
-		: FBaseTreeNode(InName, bInIsGroup)
+	SEntityEventsTableView::FEventNode::FEventNode(uint64 InProviderEntityIndex)
+		: FBaseTreeNode(false /* bInIsGroup */)
 		, ProviderEntityIndex(InProviderEntityIndex)
 	{
 	}
 
+	const FText SEntityEventsTableView::FEventNode::GetDisplayName() const
+	{
+		return FText::FromString(FString::Printf(TEXT("event_%llu"), ProviderEntityIndex));
+	}
+
 	INSIGHTS_IMPLEMENT_RTTI(SEntityEventsTableView::FEventNode);
 
 	uint64 SEntityEventsTableView::FEventNode::GetProviderEventIndex() const
@@ -365,9 +370,7 @@ namespace MassInsightsUI
 					const int32 EntityIndex = static_cast<int32>(static_cast<uint64>((EntityPtr - Entities.GetData())) / sizeof(uint64));
 					check(EntityIndex >= 0 && EntityIndex < Entities.Num());
 
-					const FName GroupName(TEXT("All"));
-					const bool bIsInGroup = false;
-					TSharedPtr<FEventNode> EventNode = MakeShared<FEventNode>(GroupName, bIsInGroup, ProviderEventIndex);
+					TSharedPtr<FEventNode> EventNode = MakeShared<FEventNode>(ProviderEventIndex);
 					Events.Emplace(MoveTemp(EventNode));
 				});
 				
diff --git a/Engine/Plugins/MassInsights/Source/MassInsightsUI/Private/MassInsightsUI/Widgets/SEntityEventsTableView.h b/Engine/Plugins/MassInsights/Source/MassInsightsUI/Private/MassInsightsUI/Widgets/SEntityEventsTableView.h
--- a/Engine/Plugins/MassInsights/Source/MassInsightsUI/Private/MassInsightsUI/Widgets/SEntityEventsTableView.h
+++ b/Engine/Plugins/MassInsights/Source/MassInsightsUI/Private/MassInsightsUI/Widgets/SEntityEventsTableView.h
@@ -44,13 +44,14 @@ namespace MassInsightsUI
 		{
 			INSIGHTS_DECLARE_RTTI(FEventNode, UE::Insights::FBaseTreeNode)
 		public:
-			explicit FEventNode(const FName InName, bool bInIsGroup, uint64 InProviderEntityIndex);
+			explicit FEventNode(uint64 InProviderEntityIndex);
+
+			virtual const FText GetDisplayName() const override;
 
 			uint64 GetProviderEventIndex() const;
 
 		private:
 			uint64 ProviderEntityIndex;
-			int32 Index;
 		};
 
 		friend class SEntityEventsTableRow;
